Fix configurations in EpicGames.Perforce.

#rnx

diff --git a/Engine/Source/Programs/DotNet/EpicGames.Perforce/EpicGames.Perforce.csproj b/Engine/Source/Programs/DotNet/EpicGames.Perforce/EpicGames.Perforce.csproj
--- a/Engine/Source/Programs/DotNet/EpicGames.Perforce/EpicGames.Perforce.csproj
+++ b/Engine/Source/Programs/DotNet/EpicGames.Perforce/EpicGames.Perforce.csproj
@@ -1,9 +1,31 @@
 <Project Sdk="Microsoft.NET.Sdk">
 
   <PropertyGroup>
-    <TargetFramework>netcoreapp3.1</TargetFramework>
+    <TargetFramework>netstandard2.0</TargetFramework>
+    <Configuration Condition=" '$(Configuration)' == '' ">Development</Configuration>
+    <OutputType>Library</OutputType>
+    <RootNamespace>EpicGames.Common</RootNamespace>
+    <GenerateAssemblyInfo>false</GenerateAssemblyInfo>
+    <GenerateTargetFrameworkAttribute>false</GenerateTargetFrameworkAttribute>
+    <Configurations>Debug;Release;Development</Configurations>
     <Nullable>enable</Nullable>
   </PropertyGroup>
+  <PropertyGroup Condition=" '$(Configuration)|$(Platform)' == 'Debug|AnyCPU' ">
+    <DocumentationFile></DocumentationFile>
+  </PropertyGroup>
+  <PropertyGroup Condition="'$(Configuration)|$(Platform)' == 'Development|AnyCPU'">
+    <DocumentationFile></DocumentationFile>
+    <Optimize>true</Optimize>
+  </PropertyGroup>
+  <PropertyGroup Condition="'$(Configuration)|$(Platform)' == 'Release|AnyCPU'">
+    <DocumentationFile></DocumentationFile>
+    <Optimize>true</Optimize>
+  </PropertyGroup>
+  <ItemGroup>
+    <Compile Include="..\..\DotNETCommon\MetaData.cs">
+      <Link>Properties\MetaData.cs</Link>
+    </Compile>
+  </ItemGroup>
 
   <ItemGroup>
     <PackageReference Include="Datadog.Trace" Version="1.18.0" />
