Forking DotNETUtilities into EpicGames.Common, to allow unification with Horde. Uses Core-style hierarchical naming, and excludes classes that have counterparts in NET Core. DotNETUtilities will be deprecated/made as a wrapper around EpicGames.Common.

#rnx

diff --git a/Engine/Source/Programs/DotNet/EpicGames.Common/AssemblyUtils.cs b/Engine/Source/Programs/DotNet/EpicGames.Common/AssemblyUtils.cs
new file mode 100644
--- /dev/null
+++ b/Engine/Source/Programs/DotNet/EpicGames.Common/AssemblyUtils.cs
@@ -0,0 +1,134 @@
+// Copyright Epic Games, Inc. All Rights Reserved.
+
+using System;
+using System.Collections.Generic;
+using System.Diagnostics;
+using System.IO;
+using System.Linq;
+using System.Reflection;
+
+namespace EpicGames.Common
+{
+	public static class AssemblyUtils
+	{
+        /// <summary>
+        /// Gets the original location (path and filename) of an assembly.
+        /// This method is using Assembly.CodeBase property to properly resolve original
+        /// assembly path in case shadow copying is enabled.
+        /// </summary>
+        /// <returns>Absolute path and filename to the assembly.</returns>
+        public static string GetOriginalLocation(this Assembly ThisAssembly)
+        {
+            return new Uri(ThisAssembly.CodeBase).LocalPath;
+        }
+    
+        /// <summary>
+        /// Version info of the executable which runs this code.
+        /// </summary>
+        public static FileVersionInfo ExecutableVersion
+        {
+            get
+            {
+                return FileVersionInfo.GetVersionInfo(Assembly.GetEntryAssembly().GetOriginalLocation());
+            }
+        }
+
+        /// <summary>
+        /// Installs an assembly resolver. Mostly used to get shared assemblies that we don't want copied around to various output locations as happens when "Copy Local" is set to true
+        /// for an assembly reference (which is the default).
+        /// </summary>
+        public static void InstallAssemblyResolver(string PathToBinariesDotNET)
+        {
+			AppDomain.CurrentDomain.AssemblyResolve += (sender, args) =>
+            {
+			    // Name is fully qualified assembly definition - e.g. "p4dn, Version=1.0.0.0, Culture=neutral, PublicKeyToken=ff968dc1933aba6f"
+			    string AssemblyName = args.Name.Split(',')[0];
+
+				return (
+                    from KnownAssemblyName in new[] { "SwarmAgent.exe", "../ThirdParty/Ionic/Ionic.Zip.Reduced.dll", "../ThirdParty/Newtonsoft/NewtonSoft.Json.dll" }
+                    where AssemblyName.Equals(Path.GetFileNameWithoutExtension(KnownAssemblyName), StringComparison.InvariantCultureIgnoreCase)
+                    let ResolvedAssemblyFilename = Path.Combine(PathToBinariesDotNET, KnownAssemblyName)
+                    // check if the file exists first. If we just try to load it, we correctly throw an exception, but it's a generic
+                    // FileNotFoundException, which is not informative. Better to return null.
+                    select File.Exists(ResolvedAssemblyFilename) ? Assembly.LoadFile(ResolvedAssemblyFilename) : null
+                    ).FirstOrDefault();
+			};
+        }
+
+		/// <summary>
+		/// Installs an assembly resolver, which will load *any* assembly which exists recursively within the supplied folder.
+		/// </summary>
+		/// <param name="RootDirectory">The directory to enumerate.</param>
+		public static void InstallRecursiveAssemblyResolver(string RootDirectory)
+		{
+			RefreshAssemblyCache(RootDirectory);
+			
+			AppDomain.CurrentDomain.AssemblyResolve += (sender, args) =>
+			{
+				// Name is fully qualified assembly definition - e.g. "p4dn, Version=1.0.0.0, Culture=neutral, PublicKeyToken=ff968dc1933aba6f"
+				string AssemblyName = args.Name.Split(',')[0];
+				string AssemblyLocation;
+				if (AssemblyLocationCache.TryGetValue(AssemblyName, out AssemblyLocation))
+				{
+					// We have this assembly in our folder.					
+					if (File.Exists(AssemblyLocation))
+					{
+						// The assembly still exists, so load it.
+						return Assembly.LoadFile(AssemblyLocation);
+					}
+					else
+					{
+						// The assembly no longer exists on disk, so remove it from our cache.
+						AssemblyLocationCache.Remove(AssemblyName);
+					}
+				}
+
+				// The assembly wasn't found, though may have been compiled or copied as a dependency
+				RefreshAssemblyCache(RootDirectory, string.Format("{0}.dll", AssemblyName));
+				if (AssemblyLocationCache.TryGetValue(AssemblyName, out AssemblyLocation))
+				{
+					return Assembly.LoadFile(AssemblyLocation);
+				}
+
+				return null;
+			};
+		}
+
+		private static void RefreshAssemblyCache(string RootDirectory, string Pattern = "*.dll")
+		{
+			// Initialize our cache of assemblies by enumerating all files in the given folder.
+			foreach (string DiscoveredAssembly in Directory.EnumerateFiles(RootDirectory, Pattern, SearchOption.AllDirectories))
+			{
+				AddFileToAssemblyCache(DiscoveredAssembly);
+			}
+
+		}
+
+		public static void AddFileToAssemblyCache(string AssemblyPath)
+		{
+			string AssemblyName = Path.GetFileNameWithoutExtension(AssemblyPath);
+			DateTime AssemblyLastWriteTime = File.GetLastWriteTimeUtc(AssemblyPath);
+			if (AssemblyLocationCache.ContainsKey(AssemblyName))
+			{
+				// We already have this assembly in our cache. Only replace it if the discovered file is newer (to avoid stale assemblies breaking stuff).
+				if (AssemblyLastWriteTime > AssemblyWriteTimes[AssemblyName])
+				{
+					AssemblyLocationCache[AssemblyName] = AssemblyPath;
+					AssemblyWriteTimes[AssemblyName] = AssemblyLastWriteTime;
+				}
+			}
+			else
+			{
+				// This is the first copy of this assembly ... add it to our cache.
+				AssemblyLocationCache.Add(AssemblyName, AssemblyPath);
+				AssemblyWriteTimes.Add(AssemblyName, AssemblyLastWriteTime);
+			}
+		}
+
+		// Map of assembly name to path on disk
+		private static Dictionary<string, string> AssemblyLocationCache = new Dictionary<string, string>(StringComparer.OrdinalIgnoreCase);
+		// Track last modified date of each assembly, so we can ensure we always reference the latest one in the case of stale assemblies on disk.
+		private static Dictionary<string, DateTime> AssemblyWriteTimes = new Dictionary<string, DateTime>(StringComparer.OrdinalIgnoreCase);
+
+    }
+}
diff --git a/Engine/Source/Programs/DotNet/EpicGames.Common/BinaryReaderExtensions.cs b/Engine/Source/Programs/DotNet/EpicGames.Common/BinaryReaderExtensions.cs
new file mode 100644
--- /dev/null
+++ b/Engine/Source/Programs/DotNet/EpicGames.Common/BinaryReaderExtensions.cs
@@ -0,0 +1,210 @@
+// Copyright Epic Games, Inc. All Rights Reserved.
+
+using System;
+using System.Collections.Generic;
+using System.IO;
+using System.Linq;
+using System.Text;
+using System.Threading.Tasks;
+
+namespace EpicGames.Common
+{
+	public static class BinaryReaderExtensions
+	{
+		/// <summary>
+		/// Read an objects from a binary reader
+		/// </summary>
+		/// <typeparam name="T">The element type for the array</typeparam>
+		/// <param name="Reader">Reader to read data from</param>
+		/// <returns>Object instance.</returns>
+		public static T ReadObject<T>(this BinaryReader Reader) where T : class, IBinarySerializable
+		{
+			return (T)Activator.CreateInstance(typeof(T), Reader);
+		}
+
+		/// <summary>
+		/// Read an array of strings from a binary reader
+		/// </summary>
+		/// <param name="Reader">Reader to read data from</param>
+		/// <returns>Array of strings, as serialized. May be null.</returns>
+		public static string[] ReadStringArray(this BinaryReader Reader)
+		{
+			return Reader.ReadArray(() => Reader.ReadString());
+		}
+
+		/// <summary>
+		/// Read an array of objects from a binary reader
+		/// </summary>
+		/// <typeparam name="T">The element type for the array</typeparam>
+		/// <param name="Reader">Reader to read data from</param>
+		/// <returns>Array of objects, as serialized. May be null.</returns>
+		public static T[] ReadArray<T>(this BinaryReader Reader) where T : class, IBinarySerializable
+		{
+			return ReadArray<T>(Reader, () => Reader.ReadObject<T>());
+		}
+
+		/// <summary>
+		/// Read an array of objects from a binary reader
+		/// </summary>
+		/// <typeparam name="T">The element type for the array</typeparam>
+		/// <param name="Reader">Reader to read data from</param>
+		/// <param name="ReadElement">Delegate to call to serialize each element</param>
+		/// <returns>Array of objects, as serialized. May be null.</returns>
+		public static T[] ReadArray<T>(this BinaryReader Reader, Func<T> ReadElement)
+		{
+			int NumItems = Reader.ReadInt32();
+			if(NumItems < 0)
+			{
+				return null;
+			}
+
+			T[] Items = new T[NumItems];
+			for(int Idx = 0; Idx < NumItems; Idx++)
+			{
+				Items[Idx] = ReadElement();
+			}
+			return Items;
+		}
+
+		/// <summary>
+		/// Read a list of strings from a binary reader
+		/// </summary>
+		/// <param name="Reader">Reader to read data from</param>
+		/// <returns>Array of strings, as serialized. May be null.</returns>
+		public static List<string> ReadStringList(this BinaryReader Reader)
+		{
+			return Reader.ReadList(() => Reader.ReadString());
+		}
+
+		/// <summary>
+		/// Read a list of objects from a binary reader
+		/// </summary>
+		/// <typeparam name="T">The element type for the list</typeparam>
+		/// <param name="Reader">Reader to read data from</param>
+		/// <param name="ReadElement">Delegate to call to serialize each element</param>
+		/// <returns>List of objects, as serialized. May be null.</returns>
+		public static List<T> ReadList<T>(this BinaryReader Reader) where T : class, IBinarySerializable
+		{
+			return ReadList<T>(Reader, () => Reader.ReadObject<T>());
+		}
+
+		/// <summary>
+		/// Read a list of objects from a binary reader
+		/// </summary>
+		/// <typeparam name="T">The element type for the list</typeparam>
+		/// <param name="Reader">Reader to read data from</param>
+		/// <param name="ReadElement">Delegate to call to serialize each element</param>
+		/// <returns>List of objects, as serialized. May be null.</returns>
+		public static List<T> ReadList<T>(this BinaryReader Reader, Func<T> ReadElement)
+		{
+			int NumItems = Reader.ReadInt32();
+			if(NumItems < 0)
+			{
+				return null;
+			}
+
+			List<T> Items = new List<T>(NumItems);
+			for(int Idx = 0; Idx < NumItems; Idx++)
+			{
+				Items.Add(ReadElement());
+			}
+			return Items;
+		}
+
+		/// <summary>
+		/// Read a list of objects from a binary reader
+		/// </summary>
+		/// <typeparam name="T">The element type for the list</typeparam>
+		/// <param name="Reader">Reader to read data from</param>
+		/// <param name="ReadKey">Delegate to call to serialize each key</param>
+		/// <param name="ReadValue">Delegate to call to serialize each value</param>
+		/// <returns>List of objects, as serialized. May be null.</returns>
+		public static Dictionary<K, V> ReadDictionary<K, V>(this BinaryReader Reader, Func<K> ReadKey, Func<V> ReadValue)
+		{
+			int NumItems = Reader.ReadInt32();
+			if(NumItems < 0)
+			{
+				return null;
+			}
+
+			Dictionary<K, V> Items = new Dictionary<K, V>(NumItems);
+			for(int Idx = 0; Idx < NumItems; Idx++)
+			{
+				K Key = ReadKey();
+				V Value = ReadValue();
+				Items.Add(Key, Value);
+			}
+			return Items;
+		}
+
+		/// <summary>
+		/// Read a nullable object from a binary reader
+		/// </summary>
+		/// <typeparam name="T">Type of the object</typeparam>
+		/// <param name="Reader">Reader to read data from</param>
+		/// <param name="ReadItem">Function to read the payload, if non-null</param>
+		/// <returns>Object instance or null</returns>
+		public static T ReadNullable<T>(this BinaryReader Reader, Func<T> ReadItem) where T : class
+		{
+			if(Reader.ReadBoolean())
+			{
+				return ReadItem();
+			}
+			else
+			{
+				return null;
+			}
+		}
+
+		/// <summary>
+		/// Reads a value of a specific type from a binary reader
+		/// </summary>
+		/// <param name="Reader">Reader for input data</param>
+		/// <param name="ObjectType">Type of value to read</param>
+		/// <returns>The value read from the stream</returns>
+		public static object ReadObject(this BinaryReader Reader, Type ObjectType)
+		{
+			if(ObjectType == typeof(string))
+			{
+				return Reader.ReadString();
+			}
+			else if(ObjectType == typeof(bool))
+			{
+				return Reader.ReadBoolean();
+			}
+			else if(ObjectType == typeof(int))
+			{
+				return Reader.ReadInt32();
+			}
+			else if(ObjectType == typeof(float))
+			{
+				return Reader.ReadSingle();
+			}
+			else if(ObjectType == typeof(double))
+			{
+				return Reader.ReadDouble();
+			}
+			else if(ObjectType == typeof(string[]))
+			{
+				return Reader.ReadStringArray();
+			}
+			else if(ObjectType == typeof(bool?))
+			{
+				int Value = Reader.ReadInt32();
+				return (Value == -1)? (bool?)null : (Value == 0)? (bool?)false : (bool?)true;
+			}
+			else if (ObjectType == typeof(FileReference))
+			{
+				return Reader.ReadFileReference();
+			}
+			else if(ObjectType.IsEnum)
+			{
+				return Enum.ToObject(ObjectType, Reader.ReadInt32());
+			}
+			else
+			{
+				throw new Exception(String.Format("Reading binary objects of type '{0}' is not currently supported.", ObjectType.Name));
+			}
+		}
+	}
+}
diff --git a/Engine/Source/Programs/DotNet/EpicGames.Common/BinaryWriterExtensions.cs b/Engine/Source/Programs/DotNet/EpicGames.Common/BinaryWriterExtensions.cs
new file mode 100644
--- /dev/null
+++ b/Engine/Source/Programs/DotNet/EpicGames.Common/BinaryWriterExtensions.cs
@@ -0,0 +1,223 @@
+// Copyright Epic Games, Inc. All Rights Reserved.
+
+using System;
+using System.Collections.Generic;
+using System.IO;
+using System.Linq;
+using System.Text;
+using System.Threading.Tasks;
+
+namespace EpicGames.Common
+{
+	/// <summary>
+	/// Interface that can be implemented by objects to support BinaryWriter container functionality. Also implies that
+	/// derived objects will have a constructor taking a single BinaryReader argument (similar to how the system ISerializable
+	/// interface assumes a specific constructor)
+	/// </summary>
+	public interface IBinarySerializable
+	{
+		void Write(BinaryWriter Writer);
+	}
+
+	/// <summary>
+	/// Extension methods for BinaryWriter class
+	/// </summary>
+	public static class BinaryWriterExtensions
+	{
+		/// <summary>
+		/// Writes an item implementing the IBinarySerializable interface to a binary writer. Included for symmetry with standard Writer.Write(X) calls.
+		/// </summary>
+		/// <param name="Writer">Writer to serialize to</param>
+		/// <param name="Item">The item to write</param>
+		public static void Write(this BinaryWriter Writer, IBinarySerializable Item)
+		{
+			Item.Write(Writer);
+		}
+
+		/// <summary>
+		/// Writes an array of strings to a binary writer.
+		/// </summary>
+		/// <param name="Writer">Writer to serialize to</param>
+		/// <param name="Items">Array of items</param>
+		public static void Write(this BinaryWriter Writer, string[] Items)
+		{
+			Write(Writer, Items, Item => Writer.Write(Item));
+		}
+
+		/// <summary>
+		/// Writes an array to a binary writer.
+		/// </summary>
+		/// <typeparam name="T">The array element type</typeparam>
+		/// <param name="Writer">Writer to serialize to</param>
+		/// <param name="Items">Array of items</param>
+		public static void Write<T>(this BinaryWriter Writer, T[] Items) where T : class, IBinarySerializable
+		{
+			Write(Writer, Items, Item => Writer.Write(Item));
+		}
+
+		/// <summary>
+		/// Writes an array to a binary writer.
+		/// </summary>
+		/// <typeparam name="T">The array element type</typeparam>
+		/// <param name="Writer">Writer to serialize to</param>
+		/// <param name="Items">Array of items</param>
+		/// <param name="WriteElement">Delegate to call to serialize each element</param>
+		public static void Write<T>(this BinaryWriter Writer, T[] Items, Action<T> WriteElement)
+		{
+			if(Items == null)
+			{
+				Writer.Write(-1);
+			}
+			else
+			{
+				Writer.Write(Items.Length);
+				for(int Idx = 0; Idx < Items.Length; Idx++)
+				{
+					WriteElement(Items[Idx]);
+				}
+			}
+		}
+
+		/// <summary>
+		/// Writes a list of strings to a binary writer.
+		/// </summary>
+		/// <param name="Writer">Writer to serialize to</param>
+		/// <param name="Items">Array of items</param>
+		public static void Write(this BinaryWriter Writer, List<string> Items)
+		{
+			Write(Writer, Items, Item => Writer.Write(Item));
+		}
+
+		/// <summary>
+		/// Writes a list to a binary writer.
+		/// </summary>
+		/// <typeparam name="T">The array element type</typeparam>
+		/// <param name="Writer">Writer to serialize to</param>
+		/// <param name="Items">Array of items</param>
+		public static void Write<T>(this BinaryWriter Writer, List<T> Items) where T : class, IBinarySerializable
+		{
+			Write(Writer, Items, Item => Writer.Write(Item));
+		}
+
+		/// <summary>
+		/// Writes a list to a binary writer.
+		/// </summary>
+		/// <typeparam name="T">The list element type</typeparam>
+		/// <param name="Writer">Writer to serialize to</param>
+		/// <param name="Items">List of items</param>
+		/// <param name="WriteElement">Delegate to call to serialize each element</param>
+		public static void Write<T>(this BinaryWriter Writer, List<T> Items, Action<T> WriteElement)
+		{
+			if (Items == null)
+			{
+				Writer.Write(-1);
+			}
+			else
+			{
+				Writer.Write(Items.Count);
+				for (int Idx = 0; Idx < Items.Count; Idx++)
+				{
+					WriteElement(Items[Idx]);
+				}
+			}
+		}
+
+		/// <summary>
+		/// Write a dictionary to a binary writer
+		/// </summary>
+		/// <typeparam name="K">The key type for the dictionary</typeparam>
+		/// <typeparam name="V">The value type for the dictionary</typeparam>
+		/// <param name="Reader">Reader to read data from</param>
+		/// <param name="Items">List of items to be written</param>
+		/// <param name="WriteKey">Delegate to call to serialize each key</param>
+		/// <param name="WriteKey">Delegate to call to serialize each value</param>
+		/// <returns>Dictionary of objects, as serialized. May be null.</returns>
+		public static void Write<K, V>(this BinaryWriter Writer, Dictionary<K, V> Items, Action<K> WriteKey, Action<V> WriteValue)
+		{
+			if(Items == null)
+			{
+				Writer.Write(-1);
+			}
+			else
+			{
+				Writer.Write(Items.Count);
+				foreach(KeyValuePair<K, V> Item in Items)
+				{
+					WriteKey(Item.Key);
+					WriteValue(Item.Value);
+				}
+			}
+		}
+
+		/// <summary>
+		/// Read a nullable object from a binary reader
+		/// </summary>
+		/// <typeparam name="T">Type of the object</typeparam>
+		/// <param name="Writer">Reader to read data from</param>
+		/// <param name="WriteItem">Function to read the payload, if non-null</param>
+		/// <returns>Object instance or null</returns>
+		public static void WriteNullable<T>(this BinaryWriter Writer, T Item, Action WriteItem) where T : class
+		{
+			if(Item == null)
+			{
+				Writer.Write(false);
+			}
+			else
+			{
+				Writer.Write(true);
+				WriteItem();
+			}
+		}
+
+		/// <summary>
+		/// Writes a value of a specific type to a binary writer
+		/// </summary>
+		/// <param name="Writer">Writer for output data</param>
+		/// <param name="FieldType">Type of value to write</param>
+		/// <param name="Value">The value to output</param>
+		public static void Write(this BinaryWriter Writer, Type FieldType, object Value)
+		{
+			if(FieldType == typeof(string))
+			{
+				Writer.Write((string)Value);
+			}
+			else if(FieldType == typeof(bool))
+			{
+				Writer.Write((bool)Value);
+			}
+			else if(FieldType == typeof(int))
+			{
+				Writer.Write((int)Value);
+			}
+			else if(FieldType == typeof(float))
+			{
+				Writer.Write((float)Value);
+			}
+			else if(FieldType == typeof(double))
+			{
+				Writer.Write((double)Value);
+			}
+			else if(FieldType.IsEnum)
+			{
+				Writer.Write((int)Value);
+			}
+			else if(FieldType == typeof(string[]))
+			{
+				Writer.Write((string[])Value);
+			}
+			else if(FieldType == typeof(bool?))
+			{
+				bool? NullableValue = (bool?)Value;
+				Writer.Write(NullableValue.HasValue? NullableValue.Value? 1 : 0 : -1);
+			}
+			else if (FieldType == typeof(FileReference))
+			{
+				Writer.Write((FileReference)Value);
+			}
+			else
+			{
+				throw new Exception(String.Format("Unsupported type '{0}' for binary serialization", FieldType.Name));
+			}
+		}
+	}
+}
diff --git a/Engine/Source/Programs/DotNet/EpicGames.Common/DirectoryReference.cs b/Engine/Source/Programs/DotNet/EpicGames.Common/DirectoryReference.cs
new file mode 100644
--- /dev/null
+++ b/Engine/Source/Programs/DotNet/EpicGames.Common/DirectoryReference.cs
@@ -0,0 +1,494 @@
+﻿// Copyright Epic Games, Inc. All Rights Reserved.
+
+using System;
+using System.Collections.Generic;
+using System.IO;
+using System.Linq;
+using System.Text;
+using System.Threading.Tasks;
+
+namespace EpicGames.Common
+{
+	/// <summary>
+	/// Representation of an absolute directory path. Allows fast hashing and comparisons.
+	/// </summary>
+	[Serializable]
+	public class DirectoryReference : FileSystemReference, IEquatable<DirectoryReference>
+	{
+		/// <summary>
+		/// Special value used to invoke the non-sanitizing constructor overload
+		/// </summary>
+		public enum Sanitize
+		{
+			None
+		}
+
+		/// <summary>
+		/// Default constructor.
+		/// </summary>
+		/// <param name="InPath">Path to this directory.</param>
+		public DirectoryReference(string InPath)
+			: base(FixTrailingPathSeparator(Path.GetFullPath(InPath)))
+		{
+		}
+
+		/// <summary>
+		/// Construct a DirectoryReference from a DirectoryInfo object.
+		/// </summary>
+		/// <param name="InInfo">Path to this file</param>
+		public DirectoryReference(DirectoryInfo InInfo)
+			: base(FixTrailingPathSeparator(InInfo.FullName))
+		{
+		}
+
+		/// <summary>
+		/// Constructor for creating a directory object directly from two strings.
+		/// </summary>
+		/// <param name="InFullName">The full, sanitized path name</param>
+		/// <param name="InSanitize">Dummy argument used to resolve this overload</param>
+		public DirectoryReference(string InFullName, Sanitize InSanitize)
+			: base(InFullName)
+		{
+		}
+
+		/// <summary>
+		/// Ensures that the correct trailing path separator is appended. On Windows, the root directory (eg. C:\) always has a trailing path separator, but no other
+		/// path does.
+		/// </summary>
+		/// <param name="DirName">Absolute path to the directory</param>
+		/// <returns>Path to the directory, with the correct trailing path separator</returns>
+		private static string FixTrailingPathSeparator(string DirName)
+		{
+			if(DirName.Length == 2 && DirName[1] == ':')
+			{
+				return DirName + Path.DirectorySeparatorChar;
+			}
+			else if(DirName.Length == 3 && DirName[1] == ':' && DirName[2] == Path.DirectorySeparatorChar)
+			{
+				return DirName;
+			}
+			else if(DirName.Length > 1 && DirName[DirName.Length - 1] == Path.DirectorySeparatorChar)
+			{
+				return DirName.TrimEnd(Path.DirectorySeparatorChar);
+			}
+			else
+			{
+				return DirName;
+			}
+		}
+
+		/// <summary>
+		/// Gets the top level directory name
+		/// </summary>
+		/// <returns>The name of the directory</returns>
+		public string GetDirectoryName()
+		{
+			return Path.GetFileName(FullName);
+		}
+
+		/// <summary>
+		/// Gets the directory containing this object
+		/// </summary>
+		/// <returns>A new directory object representing the directory containing this object</returns>
+		public DirectoryReference ParentDirectory
+		{
+			get
+			{
+				if (IsRootDirectory())
+				{
+					return null;
+				}
+
+				int ParentLength = FullName.LastIndexOf(Path.DirectorySeparatorChar);
+				if (ParentLength == 2 && FullName[1] == ':')
+				{
+					ParentLength++;
+				}
+
+				if (ParentLength == 0 && FullName[0] == Path.DirectorySeparatorChar)
+				{
+					// we have reached the root
+					ParentLength = 1;
+				}
+
+				return new DirectoryReference(FullName.Substring(0, ParentLength), Sanitize.None);
+			}
+		}
+
+		/// <summary>
+		/// Gets the parent directory for a file
+		/// </summary>
+		/// <param name="File">The file to get directory for</param>
+		/// <returns>The full directory name containing the given file</returns>
+		[Obsolete("Replace with call to FileReference.ParentDirectory instead.")]
+		public static DirectoryReference GetParentDirectory(FileReference File)
+		{
+			int ParentLength = File.FullName.LastIndexOf(Path.DirectorySeparatorChar);
+			if(ParentLength == 2 && File.FullName[1] == ':')
+			{
+				ParentLength++;
+			}
+
+			return new DirectoryReference(File.FullName.Substring(0, ParentLength), Sanitize.None);
+		}
+
+		/// <summary>
+		/// Gets the path for a special folder
+		/// </summary>
+		/// <param name="Folder">The folder to receive the path for</param>
+		/// <returns>Directory reference for the given folder, or null if it is not available</returns>
+		public static DirectoryReference GetSpecialFolder(Environment.SpecialFolder Folder)
+		{
+			string FolderPath = Environment.GetFolderPath(Folder);
+			return String.IsNullOrEmpty(FolderPath)? null : new DirectoryReference(FolderPath);
+		}
+
+		/// <summary>
+		/// Determines whether this path represents a root directory in the filesystem
+		/// </summary>
+		/// <returns>True if this path is a root directory, false otherwise</returns>
+		public bool IsRootDirectory()
+		{
+			return FullName[FullName.Length - 1] == Path.DirectorySeparatorChar;
+		}
+
+		/// <summary>
+		/// Combine several fragments with a base directory, to form a new directory name
+		/// </summary>
+		/// <param name="BaseDirectory">The base directory</param>
+		/// <param name="Fragments">Fragments to combine with the base directory</param>
+		/// <returns>The new directory name</returns>
+		public static DirectoryReference Combine(DirectoryReference BaseDirectory, params string[] Fragments)
+		{
+			string FullName = FileSystemReference.CombineStrings(BaseDirectory, Fragments);
+			return new DirectoryReference(FullName, Sanitize.None);
+		}
+
+		/// <summary>
+		/// Compares two filesystem object names for equality. Uses the canonical name representation, not the display name representation.
+		/// </summary>
+		/// <param name="A">First object to compare.</param>
+		/// <param name="B">Second object to compare.</param>
+		/// <returns>True if the names represent the same object, false otherwise</returns>
+		public static bool operator ==(DirectoryReference A, DirectoryReference B)
+		{
+			if ((object)A == null)
+			{
+				return (object)B == null;
+			}
+			else
+			{
+				return (object)B != null && A.FullName.Equals(B.FullName, Comparison);
+			}
+		}
+
+		/// <summary>
+		/// Compares two filesystem object names for inequality. Uses the canonical name representation, not the display name representation.
+		/// </summary>
+		/// <param name="A">First object to compare.</param>
+		/// <param name="B">Second object to compare.</param>
+		/// <returns>False if the names represent the same object, true otherwise</returns>
+		public static bool operator !=(DirectoryReference A, DirectoryReference B)
+		{
+			return !(A == B);
+		}
+
+		/// <summary>
+		/// Compares against another object for equality.
+		/// </summary>
+		/// <param name="Obj">other instance to compare.</param>
+		/// <returns>True if the names represent the same object, false otherwise</returns>
+		public override bool Equals(object Obj)
+		{
+			return (Obj is DirectoryReference) && ((DirectoryReference)Obj) == this;
+		}
+
+		/// <summary>
+		/// Compares against another object for equality.
+		/// </summary>
+		/// <param name="Obj">other instance to compare.</param>
+		/// <returns>True if the names represent the same object, false otherwise</returns>
+		public bool Equals(DirectoryReference Obj)
+		{
+			return Obj == this;
+		}
+
+		/// <summary>
+		/// Returns a hash code for this object
+		/// </summary>
+		/// <returns></returns>
+		public override int GetHashCode()
+		{
+			return Comparer.GetHashCode(FullName);
+		}
+
+		/// <summary>
+		/// Helper function to create a remote directory reference. Unlike normal DirectoryReference objects, these aren't converted to a full path in the local filesystem.
+		/// </summary>
+		/// <param name="AbsolutePath">The absolute path in the remote file system</param>
+		/// <returns>New directory reference</returns>
+		public static DirectoryReference MakeRemote(string AbsolutePath)
+		{
+			return new DirectoryReference(AbsolutePath, Sanitize.None);
+		}
+
+		/// <summary>
+		/// Gets the parent directory for a file, or returns null if it's null.
+		/// </summary>
+		/// <param name="File">The file to create a directory reference for</param>
+		/// <returns>The directory containing the file  </returns>
+		public static DirectoryReference FromFile(FileReference File)
+		{
+			if(File == null)
+			{
+				return null;
+			}
+			else
+			{
+				return File.Directory;
+			}
+		}
+
+		/// <summary>
+		/// Create a DirectoryReference from a string. If the string is null, returns a null DirectoryReference.
+		/// </summary>
+		/// <param name="DirectoryName">Path for the new object</param>
+		/// <returns>Returns a FileReference representing the given string, or null.</returns>
+		public static DirectoryReference FromString(string DirectoryName)
+		{
+			if(String.IsNullOrEmpty(DirectoryName))
+			{
+				return null;
+			}
+			else
+			{
+				return new DirectoryReference(DirectoryName);
+			}
+		}
+
+		/// <summary>
+		/// Finds the correct case to match the location of this file on disk. Uses the given case for parts of the path that do not exist.
+		/// </summary>
+		/// <param name="Location">The path to find the correct case for</param>
+		/// <returns>Location of the file with the correct case</returns>
+		public static DirectoryReference FindCorrectCase(DirectoryReference Location)
+		{
+			return new DirectoryReference(DirectoryUtils.FindCorrectCase(Location.ToDirectoryInfo()));
+		}
+
+		/// <summary>
+		/// Constructs a DirectoryInfo object from this reference
+		/// </summary>
+		/// <returns>New DirectoryInfo object</returns>
+		public DirectoryInfo ToDirectoryInfo()
+		{
+			return new DirectoryInfo(FullName);
+		}
+
+		#region System.IO.Directory Wrapper Methods
+
+		/// <summary>
+		/// Finds the current directory
+		/// </summary>
+		/// <returns>The current directory</returns>
+		public static DirectoryReference GetCurrentDirectory()
+		{
+			return new DirectoryReference(Directory.GetCurrentDirectory());
+		}
+
+		/// <summary>
+		/// Creates a directory
+		/// </summary>
+		/// <param name="Location">Location of the directory</param>
+		public static void CreateDirectory(DirectoryReference Location)
+		{
+			Directory.CreateDirectory(Location.FullName);
+		}
+
+        /// <summary>
+        /// Deletes a directory
+        /// </summary>
+		/// <param name="Location">Location of the directory</param>
+        public static void Delete(DirectoryReference Location)
+        {
+            Directory.Delete(Location.FullName);
+        }
+
+        /// <summary>
+        /// Deletes a directory
+        /// </summary>
+		/// <param name="Location">Location of the directory</param>
+		/// <param name="bRecursive">Whether to remove directories recursively</param>
+        public static void Delete(DirectoryReference Location, bool bRecursive)
+        {
+            Directory.Delete(Location.FullName, bRecursive);
+        }
+
+        /// <summary>
+        /// Checks whether the directory exists
+        /// </summary>
+		/// <param name="Location">Location of the directory</param>
+        /// <returns>True if this directory exists</returns>
+        public static bool Exists(DirectoryReference Location)
+		{
+			return Directory.Exists(Location.FullName);
+		}
+
+		/// <summary>
+		/// Enumerate files from a given directory
+		/// </summary>
+		/// <param name="BaseDir">Base directory to search in</param>
+		/// <returns>Sequence of file references</returns>
+		public static IEnumerable<FileReference> EnumerateFiles(DirectoryReference BaseDir)
+		{
+			foreach (string FileName in Directory.EnumerateFiles(BaseDir.FullName))
+			{
+				yield return new FileReference(FileName, FileReference.Sanitize.None);
+			}
+		}
+
+		/// <summary>
+		/// Enumerate files from a given directory
+		/// </summary>
+		/// <param name="BaseDir">Base directory to search in</param>
+		/// <param name="Pattern">Pattern for matching files</param>
+		/// <returns>Sequence of file references</returns>
+		public static IEnumerable<FileReference> EnumerateFiles(DirectoryReference BaseDir, string Pattern)
+		{
+			foreach (string FileName in Directory.EnumerateFiles(BaseDir.FullName, Pattern))
+			{
+				yield return new FileReference(FileName, FileReference.Sanitize.None);
+			}
+		}
+
+		/// <summary>
+		/// Enumerate files from a given directory
+		/// </summary>
+		/// <param name="BaseDir">Base directory to search in</param>
+		/// <param name="Pattern">Pattern for matching files</param>
+		/// <param name="Option">Options for the search</param>
+		/// <returns>Sequence of file references</returns>
+		public static IEnumerable<FileReference> EnumerateFiles(DirectoryReference BaseDir, string Pattern, SearchOption Option)
+		{
+			foreach (string FileName in Directory.EnumerateFiles(BaseDir.FullName, Pattern, Option))
+			{
+				yield return new FileReference(FileName, FileReference.Sanitize.None);
+			}
+		}
+
+		/// <summary>
+		/// Enumerate subdirectories in a given directory
+		/// </summary>
+		/// <param name="BaseDir">Base directory to search in</param>
+		/// <returns>Sequence of directory references</returns>
+		public static IEnumerable<DirectoryReference> EnumerateDirectories(DirectoryReference BaseDir)
+		{
+			foreach (string DirectoryName in Directory.EnumerateDirectories(BaseDir.FullName))
+			{
+				yield return new DirectoryReference(DirectoryName, Sanitize.None);
+			}
+		}
+
+		/// <summary>
+		/// Enumerate subdirectories in a given directory
+		/// </summary>
+		/// <param name="BaseDir">Base directory to search in</param>
+		/// <param name="Pattern">Pattern for matching directories</param>
+		/// <returns>Sequence of directory references</returns>
+		public static IEnumerable<DirectoryReference> EnumerateDirectories(DirectoryReference BaseDir, string Pattern)
+		{
+			foreach (string DirectoryName in Directory.EnumerateDirectories(BaseDir.FullName, Pattern))
+			{
+				yield return new DirectoryReference(DirectoryName, Sanitize.None);
+			}
+		}
+
+		/// <summary>
+		/// Enumerate subdirectories in a given directory
+		/// </summary>
+		/// <param name="BaseDir">Base directory to search in</param>
+		/// <param name="Pattern">Pattern for matching files</param>
+		/// <param name="Option">Options for the search</param>
+		/// <returns>Sequence of directory references</returns>
+		public static IEnumerable<DirectoryReference> EnumerateDirectories(DirectoryReference BaseDir, string Pattern, SearchOption Option)
+		{
+			foreach (string DirectoryName in Directory.EnumerateDirectories(BaseDir.FullName, Pattern, Option))
+			{
+				yield return new DirectoryReference(DirectoryName, Sanitize.None);
+			}
+		}
+
+		/// <summary>
+		/// Sets the current directory
+		/// </summary>
+		/// <param name="Location">Location of the new current directory</param>
+		public static void SetCurrentDirectory(DirectoryReference Location)
+		{
+			Directory.SetCurrentDirectory(Location.FullName);
+		}
+
+		#endregion
+	}
+
+	/// <summary>
+	/// Extension methods for passing DirectoryReference arguments
+	/// </summary>
+	public static class DirectoryReferenceExtensionMethods
+	{
+		/// <summary>
+		/// Manually serialize a file reference to a binary stream.
+		/// </summary>
+		/// <param name="Writer">Binary writer to write to</param>
+		/// <param name="Directory">The directory reference to write</param>
+		public static void Write(this BinaryWriter Writer, DirectoryReference Directory)
+		{
+			Writer.Write((Directory == null) ? String.Empty : Directory.FullName);
+		}
+
+		/// <summary>
+		/// Manually deserialize a directory reference from a binary stream.
+		/// </summary>
+		/// <param name="Reader">Binary reader to read from</param>
+		/// <returns>New DirectoryReference object</returns>
+		public static DirectoryReference ReadDirectoryReference(this BinaryReader Reader)
+		{
+			string FullName = Reader.ReadString();
+			return (FullName.Length == 0) ? null : new DirectoryReference(FullName, DirectoryReference.Sanitize.None);
+		}
+
+		/// <summary>
+		/// Writes a directory reference  to a binary archive
+		/// </summary>
+		/// <param name="Writer">The writer to output data to</param>
+		/// <param name="Directory">The item to write</param>
+		public static void WriteDirectoryReference(this BinaryArchiveWriter Writer, DirectoryReference Directory)
+		{
+			if(Directory == null)
+			{
+				Writer.WriteString(null);
+			}
+			else
+			{
+				Writer.WriteString(Directory.FullName);
+			}
+		}
+
+		/// <summary>
+		/// Reads a directory reference from a binary archive
+		/// </summary>
+		/// <param name="Reader">Reader to serialize data from</param>
+		/// <returns>New directory reference instance</returns>
+		public static DirectoryReference ReadDirectoryReference(this BinaryArchiveReader Reader)
+		{
+			string FullName = Reader.ReadString();
+			if(FullName == null)
+			{
+				return null;
+			}
+			else
+			{
+				return new DirectoryReference(FullName, DirectoryReference.Sanitize.None);
+			}
+		}
+	}
+}
diff --git a/Engine/Source/Programs/DotNet/EpicGames.Common/EpicGames.Common.csproj b/Engine/Source/Programs/DotNet/EpicGames.Common/EpicGames.Common.csproj
new file mode 100644
--- /dev/null
+++ b/Engine/Source/Programs/DotNet/EpicGames.Common/EpicGames.Common.csproj
@@ -0,0 +1,30 @@
+﻿<Project Sdk="Microsoft.NET.Sdk">
+  <PropertyGroup>
+    <TargetFramework>netcoreapp3.1</TargetFramework>
+    <Configuration Condition=" '$(Configuration)' == '' ">Development</Configuration>
+    <OutputType>Library</OutputType>
+    <RootNamespace>EpicGames.Common</RootNamespace>
+    <GenerateAssemblyInfo>false</GenerateAssemblyInfo>
+    <GenerateTargetFrameworkAttribute>false</GenerateTargetFrameworkAttribute>
+    <Configurations>Debug;Release;Development</Configurations>
+  </PropertyGroup>
+  <PropertyGroup Condition=" '$(Configuration)|$(Platform)' == 'Debug|AnyCPU' ">
+    <DocumentationFile></DocumentationFile>
+  </PropertyGroup>
+  <PropertyGroup Condition="'$(Configuration)|$(Platform)' == 'Development|AnyCPU'">
+    <DocumentationFile></DocumentationFile>
+    <Optimize>true</Optimize>
+  </PropertyGroup>
+  <PropertyGroup Condition="'$(Configuration)|$(Platform)' == 'Release|AnyCPU'">
+    <DocumentationFile></DocumentationFile>
+    <Optimize>true</Optimize>
+  </PropertyGroup>
+  <ItemGroup>
+    <Compile Include="..\..\DotNETCommon\MetaData.cs">
+      <Link>Properties\MetaData.cs</Link>
+    </Compile>
+  </ItemGroup>
+  <ItemGroup>
+    <PackageReference Include="System.Memory" Version="4.5.4" />
+  </ItemGroup>
+</Project>
\ No newline at end of file
diff --git a/Engine/Source/Programs/DotNet/EpicGames.Common/FileFilter.cs b/Engine/Source/Programs/DotNet/EpicGames.Common/FileFilter.cs
new file mode 100644
--- /dev/null
+++ b/Engine/Source/Programs/DotNet/EpicGames.Common/FileFilter.cs
@@ -0,0 +1,589 @@
+// Copyright Epic Games, Inc. All Rights Reserved.
+
+using System;
+using System.Collections.Generic;
+using System.Diagnostics;
+using System.IO;
+using System.Linq;
+using System.Text;
+using System.Text.RegularExpressions;
+
+namespace EpicGames.Common
+{
+	/// <summary>
+	/// Indicates whether files which match a pattern should be included or excluded
+	/// </summary>
+	public enum FileFilterType
+	{
+		/// <summary>
+		/// Include files matching this pattern
+		/// </summary>
+		Include,
+
+		/// <summary>
+		/// Exclude files matching this pattern
+		/// </summary>
+		Exclude,
+	}
+
+	/// <summary>
+	/// Stores a set of rules, similar to p4 path specifications, which can be used to efficiently filter files in or out of a set.
+	/// The rules are merged into a tree with nodes for each path fragments and rule numbers in each leaf, allowing tree traversal in an arbitrary order 
+	/// to how they are applied.
+	/// </summary>
+	public class FileFilter
+	{
+		/// <summary>
+		/// Root node for the tree.
+		/// </summary>
+		FileFilterNode RootNode;
+
+		/// <summary>
+		/// The default node, which will match any path
+		/// </summary>
+		FileFilterNode DefaultNode;
+
+		/// <summary>
+		/// Terminating nodes for each rule added to the filter
+		/// </summary>
+		List<FileFilterNode> Rules = new List<FileFilterNode>();
+
+		/// <summary>
+		/// Default constructor
+		/// </summary>
+		public FileFilter(FileFilterType DefaultType = FileFilterType.Exclude)
+		{
+			RootNode = new FileFilterNode(null, "");
+
+			DefaultNode = new FileFilterNode(RootNode, "...");
+			DefaultNode.Type = DefaultType;
+		}
+
+		/// <summary>
+		/// Copy constructor
+		/// </summary>
+		/// <param name="Other">Filter to copy from</param>
+		public FileFilter(FileFilter Other)
+			: this(Other.DefaultNode.Type)
+		{
+			foreach (FileFilterNode OtherRule in Other.Rules)
+			{
+				AddRule(OtherRule.ToString(), OtherRule.Type);
+			}
+		}
+
+		/// <summary>
+		/// Constructs a file filter from a p4-style filespec. Exclude lines are prefixed with a - character.
+		/// </summary>
+		/// <param name="Lines">Lines to parse rules from</param>
+		public FileFilter(IEnumerable<string> Lines)
+			: this()
+		{
+			foreach (string Line in Lines)
+			{
+				AddRule(Line);
+			}
+		}
+
+		/// <summary>
+		/// Adds an include or exclude rule to the filter
+		/// </summary>
+		/// <param name="Rule">Pattern to match. See CreateRegex() for details.</param>
+		public void AddRule(string Rule)
+		{
+			if (Rule.StartsWith("-"))
+			{
+				Exclude(Rule.Substring(1).TrimStart());
+			}
+			else
+			{
+				Include(Rule);
+			}
+		}
+
+		/// <summary>
+		/// Adds an include or exclude rule to the filter. The rule may be 
+		/// </summary>
+		/// <param name="Rule">Pattern to match. See CreateRegex() for details.</param>
+		/// <param name="AllowTags">If the rule starts with a list of tags {a,b,c}, specifies a list of tags to match against</param>
+		public void AddRule(string Rule, params string[] AllowTags)
+		{
+			string CleanRule = Rule.Trim();
+			if(CleanRule.StartsWith("{"))
+			{
+				// Find the end of the condition
+				int ConditionEnd = CleanRule.IndexOf('}');
+				if(ConditionEnd == -1)
+				{
+					throw new Exception(String.Format("Missing closing parenthesis in rule: {0}", CleanRule));
+				}
+
+				// Check there's a matching tag
+				string[] RuleTags = CleanRule.Substring(1, ConditionEnd - 1).Split(',').Select(x => x.Trim()).ToArray();
+				if(!RuleTags.Any(x => AllowTags.Contains(x)))
+				{
+					return;
+				}
+
+				// Strip the condition from the rule
+				CleanRule = CleanRule.Substring(ConditionEnd + 1).TrimStart(); 
+			}
+			AddRule(CleanRule);
+		}
+
+		/// <summary>
+		/// Adds several rules to the filter
+		/// </summary>
+		/// <param name="Rules">List of patterns to match.</param>
+		public void AddRules(IEnumerable<string> Rules)
+		{
+			foreach(string Rule in Rules)
+			{
+				AddRule(Rule);
+			}
+		}
+
+		/// <summary>
+		/// Adds several rules to the filter
+		/// </summary>
+		/// <param name="Rules">List of patterns to match.</param>
+		/// <param name="Type">Whether the rules are include or exclude rules</param>
+		public void AddRules(IEnumerable<string> Rules, FileFilterType Type)
+		{
+			foreach(string Rule in Rules)
+			{
+				AddRule(Rule, Type);
+			}
+		}
+
+		/// <summary>
+		/// Adds several rules in the given lines. Rules may be prefixed with conditions of the syntax {Key=Value, Key2=Value2}, which
+		/// will be evaluated using variables in the given dictionary before being added.
+		/// </summary>
+		/// <param name="Rules"></param>
+		/// <param name="Tags">Lookup for variables to test against</param>
+		public void AddRules(IEnumerable<string> Rules, params string[] Tags)
+		{
+			foreach(string Rule in Rules)
+			{
+				AddRule(Rule, Tags);
+			}
+		}
+
+		/// <summary>
+		/// Adds a rule which matches a filename relative to a given base directory.
+		/// </summary>
+		/// <param name="File">The filename to add a rule for</param>
+		/// <param name="BaseDirectory">Base directory for the rule</param>
+		/// <param name="Type">Whether to add an include or exclude rule</param>
+		public void AddRuleForFile(FileReference File, DirectoryReference BaseDirectory, FileFilterType Type)
+		{
+			AddRule("/" + File.MakeRelativeTo(BaseDirectory), Type);
+		}
+
+		/// <summary>
+		/// Adds rules for files which match the given names
+		/// </summary>
+		/// <param name="Files">The filenames to rules for</param>
+		/// <param name="BaseDirectory">Base directory for the rules</param>
+		/// <param name="Type">Whether to add an include or exclude rule</param>
+		public void AddRuleForFiles(IEnumerable<FileReference> Files, DirectoryReference BaseDirectory, FileFilterType Type)
+		{
+			foreach (FileReference File in Files)
+			{
+				AddRuleForFile(File, BaseDirectory, Type);
+			}
+		}
+
+		/// <summary>
+		/// Reads a configuration file split into sections
+		/// </summary>
+		/// <param name="FileName"></param>
+		/// <param name="SectionName"></param>
+		/// <param name="AllowTags"></param>
+		public void ReadRulesFromFile(FileReference FileName, string SectionName, params string[] AllowTags)
+		{
+			bool bInSection = false;
+			foreach(string Line in File.ReadAllLines(FileName.FullName))
+			{
+				string TrimLine = Line.Trim();
+				if(!TrimLine.StartsWith(";") && TrimLine.Length > 0)
+				{
+					if(TrimLine.StartsWith("["))
+					{
+						bInSection = (TrimLine == "[" + SectionName + "]");
+					}
+					else if(bInSection)
+					{
+						AddRule(Line, AllowTags);
+					}
+				}
+			}
+		}
+
+		/// <summary>
+		/// Adds an include rule to the filter
+		/// </summary>
+		/// <param name="Pattern">Pattern to match. See CreateRegex() for details.</param>
+		public void Include(string Pattern)
+		{
+			AddRule(Pattern, FileFilterType.Include);
+		}
+
+		/// <summary>
+		/// Adds several exclude rules to the filter
+		/// </summary>
+		/// <param name="Patterns">Patterns to match. See CreateRegex() for details.</param>
+		public void Include(IEnumerable<string> Patterns)
+		{
+			foreach (string Pattern in Patterns)
+			{
+				Include(Pattern);
+			}
+		}
+
+		/// <summary>
+		/// Adds an exclude rule to the filter
+		/// </summary>
+		/// <param name="Pattern">Mask to match. See CreateRegex() for details.</param>
+		public void Exclude(string Pattern)
+		{
+			AddRule(Pattern, FileFilterType.Exclude);
+		}
+
+		/// <summary>
+		/// Adds several exclude rules to the filter
+		/// </summary>
+		/// <param name="Patterns">Patterns to match. See CreateRegex() for details.</param>
+		public void Exclude(IEnumerable<string> Patterns)
+		{
+			foreach (string Pattern in Patterns)
+			{
+				Exclude(Pattern);
+			}
+		}
+
+		/// <summary>
+		/// Adds an include or exclude rule to the filter
+		/// </summary>
+		/// <param name="Pattern">The pattern which the rule should match</param>
+		/// <param name="Type">Whether to include or exclude files matching this rule</param>
+		public void AddRule(string Pattern, FileFilterType Type)
+		{
+			string NormalizedPattern = Pattern.Replace('\\', '/');
+
+			// Remove the slash from the start of the pattern. Any exclude pattern that doesn't contain a directory separator is assumed to apply to any directory (eg. *.cpp), otherwise it's 
+			// taken relative to the root.
+			if (NormalizedPattern.StartsWith("/"))
+			{
+				NormalizedPattern = NormalizedPattern.Substring(1);
+			}
+			else if(!NormalizedPattern.Contains("/") && !NormalizedPattern.StartsWith("..."))
+			{
+				NormalizedPattern = ".../" + NormalizedPattern;
+			}
+
+			// All directories indicate a wildcard match
+			if (NormalizedPattern.EndsWith("/"))
+			{
+				NormalizedPattern += "...";
+			}
+
+			// Replace any directory wildcards mid-string
+			for (int Idx = NormalizedPattern.IndexOf("..."); Idx != -1; Idx = NormalizedPattern.IndexOf("...", Idx))
+			{
+				if (Idx > 0 && NormalizedPattern[Idx - 1] != '/')
+				{
+					NormalizedPattern = NormalizedPattern.Insert(Idx, "*/");
+					Idx++;
+				}
+
+				Idx += 3;
+
+				if (Idx < NormalizedPattern.Length && NormalizedPattern[Idx] != '/')
+				{
+					NormalizedPattern = NormalizedPattern.Insert(Idx, "/*");
+					Idx += 2;
+				}
+			}
+
+			// Split the pattern into fragments
+			string[] BranchPatterns = NormalizedPattern.Split('/');
+
+			// Add it into the tree
+			FileFilterNode LastNode = RootNode;
+			foreach (string BranchPattern in BranchPatterns)
+			{
+				FileFilterNode NextNode = LastNode.Branches.FirstOrDefault(x => x.Pattern == BranchPattern);
+				if (NextNode == null)
+				{
+					NextNode = new FileFilterNode(LastNode, BranchPattern);
+					LastNode.Branches.Add(NextNode);
+				}
+				LastNode = NextNode;
+			}
+
+			// We've reached the end of the pattern, so mark it as a leaf node
+			Rules.Add(LastNode);
+			LastNode.RuleNumber = Rules.Count - 1;
+			LastNode.Type = Type;
+
+			// Update the maximums along that path
+			for (FileFilterNode UpdateNode = LastNode; UpdateNode != null; UpdateNode = UpdateNode.Parent)
+			{
+				if (Type == FileFilterType.Include)
+				{
+					UpdateNode.MaxIncludeRuleNumber = LastNode.RuleNumber;
+				}
+				else
+				{
+					UpdateNode.MaxExcludeRuleNumber = LastNode.RuleNumber;
+				}
+			}
+		}
+
+		/// <summary>
+		/// Determines whether the given file matches the filter
+		/// </summary>
+		/// <param name="FileName">File to match</param>
+		/// <returns>True if the file passes the filter</returns>
+		public bool Matches(string FileName)
+		{
+			string[] Tokens = FileName.TrimStart('/', '\\').Split('/', '\\');
+
+			FileFilterNode MatchingNode = FindMatchingNode(RootNode, Tokens, 0, DefaultNode);
+
+			return MatchingNode.Type == FileFilterType.Include;
+		}
+
+		/// <summary>
+		/// Determines whether it's possible for anything within the given folder name to match the filter. Useful to early out of recursive file searches.
+		/// </summary>
+		/// <param name="FolderName">File to match</param>
+		/// <returns>True if the file passes the filter</returns>
+		public bool PossiblyMatches(string FolderName)
+		{
+			string[] Tokens = FolderName.Trim('/', '\\').Split('/', '\\');
+
+			FileFilterNode MatchingNode = FindMatchingNode(RootNode, Tokens.Union(new string[]{ "" }).ToArray(), 0, DefaultNode);
+
+			return MatchingNode.Type == FileFilterType.Include || HighestPossibleIncludeMatch(RootNode, Tokens, 0, MatchingNode.RuleNumber) > MatchingNode.RuleNumber;
+		}
+
+		/// <summary>
+		/// Applies the filter to each element in a sequence, and returns the list of files that match
+		/// </summary>
+		/// <param name="FileNames">List of filenames</param>
+		/// <returns>List of filenames which match the filter</returns>
+		public IEnumerable<string> ApplyTo(IEnumerable<string> FileNames)
+		{
+			return FileNames.Where(x => Matches(x));
+		}
+
+		/// <summary>
+		/// Applies the filter to each element in a sequence, and returns the list of files that match
+		/// </summary>
+		/// <param name="BaseDirectory">The base directory to match files against</param>
+		/// <param name="FileNames">List of filenames</param>
+		/// <returns>List of filenames which match the filter</returns>
+		public IEnumerable<FileReference> ApplyTo(DirectoryReference BaseDirectory, IEnumerable<FileReference> FileNames)
+		{
+			return FileNames.Where(x => Matches(x.MakeRelativeTo(BaseDirectory)));
+		}
+
+		/// <summary>
+		/// Finds a list of files within a given directory which match the filter.
+		/// </summary>
+		/// <param name="DirectoryName">File to match</param>
+		/// <param name="bIgnoreSymlinks">Whether to ignore symlinks in the output</param>
+		/// <returns>List of files that pass the filter</returns>
+		public List<FileReference> ApplyToDirectory(DirectoryReference DirectoryName, bool bIgnoreSymlinks)
+		{
+			List<FileReference> MatchingFileNames = new List<FileReference>();
+			if (DirectoryReference.Exists(DirectoryName))
+			{
+				FindMatchesFromDirectory(new DirectoryInfo(DirectoryName.FullName), "", bIgnoreSymlinks, MatchingFileNames);
+			}
+			return MatchingFileNames;
+		}
+
+		/// <summary>
+		/// Finds a list of files within a given directory which match the filter.
+		/// </summary>
+		/// <param name="DirectoryName">File to match</param>
+		/// <param name="PrefixPath"></param>
+		/// <param name="bIgnoreSymlinks">Whether to ignore symlinks in the output</param>
+		/// <returns>List of files that pass the filter</returns>
+		public List<FileReference> ApplyToDirectory(DirectoryReference DirectoryName, string PrefixPath, bool bIgnoreSymlinks)
+		{
+			List<FileReference> MatchingFileNames = new List<FileReference>();
+			if (DirectoryReference.Exists(DirectoryName))
+			{
+				FindMatchesFromDirectory(new DirectoryInfo(DirectoryName.FullName), PrefixPath.Replace('\\', '/').TrimEnd('/') + "/", bIgnoreSymlinks, MatchingFileNames);
+			}
+			return MatchingFileNames;
+		}
+
+		/// <summary>
+		/// Checks whether the given pattern contains a supported wildcard. Useful for distinguishing explicit file references from opportunistic file references.
+		/// </summary>
+		/// <param name="Pattern">The pattern to check</param>
+		/// <returns>True if the pattern contains a wildcard (?, *, ...), false otherwise.</returns>
+		public static int FindWildcardIndex(string Pattern)
+		{
+			int Result = -1;
+			for(int Idx = 0; Idx < Pattern.Length; Idx++)
+			{
+				if(Pattern[Idx] == '?' || Pattern[Idx] == '*' || (Pattern[Idx] == '.' && Idx + 2 < Pattern.Length && Pattern[Idx + 1] == '.' && Pattern[Idx + 2] == '.'))
+				{
+					Result = Idx;
+					break;
+				}
+			}
+			return Result;
+		}
+
+		/// <summary>
+		/// Resolve an individual wildcard to a set of files
+		/// </summary>
+		/// <param name="Pattern">The pattern to resolve</param>
+		/// <returns>List of files matching the wildcard</returns>
+		public static List<FileReference> ResolveWildcard(string Pattern)
+		{
+			for(int Idx = FindWildcardIndex(Pattern); Idx > 0; Idx--)
+			{
+				if(Pattern[Idx] == '/' || Pattern[Idx] == '\\')
+				{
+					return ResolveWildcard(new DirectoryReference(Pattern.Substring(0, Idx)), Pattern.Substring(Idx + 1));
+				}
+			}
+			return ResolveWildcard(DirectoryReference.GetCurrentDirectory(), Pattern);
+		}
+
+		/// <summary>
+		/// Resolve an individual wildcard to a set of files
+		/// </summary>
+		/// <param name="BaseDir">Base directory for wildcards</param>
+		/// <param name="Pattern">The pattern to resolve</param>
+		/// <returns>List of files matching the wildcard</returns>
+		public static List<FileReference> ResolveWildcard(DirectoryReference BaseDir, string Pattern)
+		{
+			FileFilter Filter = new FileFilter(FileFilterType.Exclude);
+			Filter.AddRule(Pattern);
+			return Filter.ApplyToDirectory(BaseDir, true);
+		}
+
+		/// <summary>
+		/// Finds a list of files within a given directory which match the filter.
+		/// </summary>
+		/// <param name="CurrentDirectory">The current directory</param>
+		/// <param name="NamePrefix">Current relative path prefix in the traversed directory tree</param>
+		/// <param name="bIgnoreSymlinks">Whether to ignore symlinks in the output</param>
+		/// <param name="MatchingFileNames">Receives a list of matching files</param>
+		/// <returns>True if the file passes the filter</returns>
+		void FindMatchesFromDirectory(DirectoryInfo CurrentDirectory, string NamePrefix, bool bIgnoreSymlinks, List<FileReference> MatchingFileNames)
+		{
+			foreach (FileInfo NextFile in CurrentDirectory.EnumerateFiles())
+			{
+				string FileName = NamePrefix + NextFile.Name;
+				if (Matches(FileName) && (!bIgnoreSymlinks || !NextFile.Attributes.HasFlag(FileAttributes.ReparsePoint)))
+				{
+					MatchingFileNames.Add(new FileReference(NextFile));
+				}
+			}
+			foreach (DirectoryInfo NextDirectory in CurrentDirectory.EnumerateDirectories())
+			{
+				string NextNamePrefix = NamePrefix + NextDirectory.Name;
+				if (PossiblyMatches(NextNamePrefix))
+				{
+					FindMatchesFromDirectory(NextDirectory, NextNamePrefix + "/", bIgnoreSymlinks, MatchingFileNames);
+				}
+			}
+		}
+
+		/// <summary>
+		/// Finds the node which matches a given list of tokens.
+		/// </summary>
+		/// <param name="CurrentNode"></param>
+		/// <param name="Tokens"></param>
+		/// <param name="TokenIdx"></param>
+		/// <param name="CurrentBestNode"></param>
+		/// <returns></returns>
+		FileFilterNode FindMatchingNode(FileFilterNode CurrentNode, string[] Tokens, int TokenIdx, FileFilterNode CurrentBestNode)
+		{
+			// If we've matched all the input tokens, check if this rule is better than any other we've seen
+			if (TokenIdx == Tokens.Length)
+			{
+				return (CurrentNode.RuleNumber > CurrentBestNode.RuleNumber) ? CurrentNode : CurrentBestNode;
+			}
+
+			// If there is no rule under the current node which is better than the current best node, early out
+			if (CurrentNode.MaxIncludeRuleNumber <= CurrentBestNode.RuleNumber && CurrentNode.MaxExcludeRuleNumber <= CurrentBestNode.RuleNumber)
+			{
+				return CurrentBestNode;
+			}
+
+			// Test all the branches for one that matches
+			FileFilterNode BestNode = CurrentBestNode;
+			foreach (FileFilterNode Branch in CurrentNode.Branches)
+			{
+				if (Branch.Pattern == "...")
+				{
+					for (int NextTokenIdx = Tokens.Length; NextTokenIdx >= TokenIdx; NextTokenIdx--)
+					{
+						BestNode = FindMatchingNode(Branch, Tokens, NextTokenIdx, BestNode);
+					}
+				}
+				else
+				{
+					if (Branch.IsMatch(Tokens[TokenIdx]))
+					{
+						BestNode = FindMatchingNode(Branch, Tokens, TokenIdx + 1, BestNode);
+					}
+				}
+			}
+			return BestNode;
+		}
+
+		/// <summary>
+		/// Returns the highest possible rule number which can match the given list of input tokens, assuming that the list of input tokens is incomplete.
+		/// </summary>
+		/// <param name="CurrentNode">The current node being checked</param>
+		/// <param name="Tokens">The tokens to match</param>
+		/// <param name="TokenIdx">Current token index</param>
+		/// <param name="CurrentBestRuleNumber">The highest rule number seen so far. Used to optimize tree traversals.</param>
+		/// <returns>New highest rule number</returns>
+		int HighestPossibleIncludeMatch(FileFilterNode CurrentNode, string[] Tokens, int TokenIdx, int CurrentBestRuleNumber)
+		{
+			// If we've matched all the input tokens, check if this rule is better than any other we've seen
+			if (TokenIdx == Tokens.Length)
+			{
+				return Math.Max(CurrentBestRuleNumber, CurrentNode.MaxIncludeRuleNumber);
+			}
+
+			// Test all the branches for one that matches
+			int BestRuleNumber = CurrentBestRuleNumber;
+			if (CurrentNode.MaxIncludeRuleNumber > BestRuleNumber)
+			{
+				foreach (FileFilterNode Branch in CurrentNode.Branches)
+				{
+					if (Branch.Pattern == "...")
+					{
+						if (Branch.MaxIncludeRuleNumber > BestRuleNumber)
+						{
+							BestRuleNumber = Branch.MaxIncludeRuleNumber;
+						}
+					}
+					else
+					{
+						if (Branch.IsMatch(Tokens[TokenIdx]))
+						{
+							BestRuleNumber = HighestPossibleIncludeMatch(Branch, Tokens, TokenIdx + 1, BestRuleNumber);
+						}
+					}
+				}
+			}
+			return BestRuleNumber;
+		}
+	}
+}
diff --git a/Engine/Source/Programs/DotNet/EpicGames.Common/FileFilterNode.cs b/Engine/Source/Programs/DotNet/EpicGames.Common/FileFilterNode.cs
new file mode 100644
--- /dev/null
+++ b/Engine/Source/Programs/DotNet/EpicGames.Common/FileFilterNode.cs
@@ -0,0 +1,120 @@
+// Copyright Epic Games, Inc. All Rights Reserved.
+
+using System;
+using System.Collections.Generic;
+using System.Linq;
+using System.Text;
+using System.Threading.Tasks;
+
+namespace EpicGames.Common
+{
+	/// <summary>
+	/// Node within a filter tree. Each node matches a single path fragment - a folder or file name, with an include or exclude rule consisting of a sequence of nodes. 
+	/// </summary>
+	class FileFilterNode
+	{
+		/// <summary>
+		/// Node which this is parented to. Null for the root node.
+		/// </summary>
+		public readonly FileFilterNode Parent;
+
+		/// <summary>
+		/// Pattern to match for this node. May contain * or ? wildcards.
+		/// </summary>
+		public readonly string Pattern;
+
+		/// <summary>
+		/// Highest include rule number matched by this node or any child nodes.
+		/// </summary>
+		public int MaxIncludeRuleNumber = -1;
+
+		/// <summary>
+		/// Highest exclude rule number matched by this node or any child nodes.
+		/// </summary>
+		public int MaxExcludeRuleNumber = -1;
+
+		/// <summary>
+		/// Child branches of this node, distinct by the the pattern for each.
+		/// </summary>
+		public List<FileFilterNode> Branches = new List<FileFilterNode>();
+
+		/// <summary>
+		/// If a path matches the sequence of nodes ending in this node, the number of the rule which added it. -1 if this was not the last node in a rule.
+		/// </summary>
+		public int RuleNumber;
+
+		/// <summary>
+		/// Whether a rule terminating in this node should include files or exclude files.
+		/// </summary>
+		public FileFilterType Type;
+
+		/// <summary>
+		/// Default constructor.
+		/// </summary>
+		public FileFilterNode(FileFilterNode InParent, string InPattern)
+		{
+			Parent = InParent;
+			Pattern = InPattern;
+			RuleNumber = -1;
+			Type = FileFilterType.Exclude;
+		}
+
+		/// <summary>
+		/// Determine if the given token matches the pattern in this node
+		/// </summary>
+		public bool IsMatch(string Token)
+		{
+			if(Pattern.EndsWith("."))
+			{
+				return !Token.Contains('.') && IsMatch(Token, 0, Pattern.Substring(0, Pattern.Length - 1), 0);
+			}
+			else
+			{
+				return IsMatch(Token, 0, Pattern, 0);
+			}
+		}
+
+		/// <summary>
+		/// Determine if a given token matches a pattern, starting from the given positions within each string.
+		/// </summary>
+		static bool IsMatch(string Token, int TokenIdx, string Pattern, int PatternIdx)
+		{
+			for (; ; )
+			{
+				if (PatternIdx == Pattern.Length)
+				{
+					return (TokenIdx == Token.Length);
+				}
+				else if (Pattern[PatternIdx] == '*')
+				{
+					for (int NextTokenIdx = Token.Length; NextTokenIdx >= TokenIdx; NextTokenIdx--)
+					{
+						if (IsMatch(Token, NextTokenIdx, Pattern, PatternIdx + 1))
+						{
+							return true;
+						}
+					}
+					return false;
+				}
+				else if (TokenIdx < Token.Length && (Char.ToLower(Token[TokenIdx]) == Char.ToLower(Pattern[PatternIdx]) || Pattern[PatternIdx] == '?'))
+				{
+					TokenIdx++;
+					PatternIdx++;
+				}
+				else
+				{
+					return false;
+				}
+			}
+		}
+
+		/// <summary>
+		/// Debugger visualization
+		/// </summary>
+		/// <returns>Path to this node</returns>
+		public override string ToString()
+		{
+			return (Parent == null) ? "/" : Parent.ToString().TrimEnd('/') + "/" + Pattern;
+		}
+	}
+}
diff --git a/Engine/Source/Programs/DotNet/EpicGames.Common/FilePattern.cs b/Engine/Source/Programs/DotNet/EpicGames.Common/FilePattern.cs
new file mode 100644
--- /dev/null
+++ b/Engine/Source/Programs/DotNet/EpicGames.Common/FilePattern.cs
@@ -0,0 +1,436 @@
+﻿// Copyright Epic Games, Inc. All Rights Reserved.
+
+using System;
+using System.Collections.Generic;
+using System.IO;
+using System.Linq;
+using System.Text;
+using System.Text.RegularExpressions;
+using System.Threading.Tasks;
+
+namespace EpicGames.Common
+{
+	public class FilePatternException : Exception
+	{
+		public FilePatternException(string Message)
+			: base(Message)
+		{
+		}
+
+		public FilePatternException(string Format, params object[] Args)
+			: base(String.Format(Format, Args))
+		{
+		}
+
+		public override string ToString()
+		{
+			return Message;
+		}
+	}
+
+	/// <summary>
+	/// Encapsulates a pattern containing the '?', '*', and '...' wildcards.
+	/// </summary>
+	public class FilePattern
+	{
+		/// <summary>
+		/// Base directory for all matched files
+		/// </summary>
+		public readonly DirectoryReference BaseDirectory;
+		
+		/// <summary>
+		/// List of tokens in the pattern. Every second token is a wildcard, other tokens are string fragments. Always has an odd number of elements. Path separators are normalized to the host platform format.
+		/// </summary>
+		public readonly List<string> Tokens = new List<string>();
+
+		/// <summary>
+		/// Constructs a file pattern which matches a single file
+		/// </summary>
+		/// <param name="File">Location of the file</param>
+		public FilePattern(FileReference File)
+		{
+			BaseDirectory = File.Directory;
+			Tokens.Add(File.GetFileName());
+		}
+
+		/// <summary>
+		/// Constructs a file pattern from the given string, resolving relative paths to the given directory. 
+		/// </summary>
+		/// <param name="RootDirectory">If a relative path is specified by the pattern, the root directory used to turn it into an absolute path</param>
+		/// <param name="Pattern">The pattern to match. If the pattern ends with a directory separator, an implicit '...' is appended.</param>
+		public FilePattern(DirectoryReference RootDirectory, string Pattern)
+		{
+			// Normalize the path separators
+			StringBuilder Text = new StringBuilder(Pattern);
+			if(Path.DirectorySeparatorChar != '\\')
+			{
+				Text.Replace('\\', Path.DirectorySeparatorChar);
+			}
+			if(Path.DirectorySeparatorChar != '/')
+			{
+				Text.Replace('/', Path.DirectorySeparatorChar);
+			}
+
+			// Find the base directory, stopping when we hit a wildcard. The source directory must end with a path specification.
+			int BaseDirectoryLen = 0;
+			for(int Idx = 0; Idx < Text.Length; Idx++)
+			{
+				if(Text[Idx] == Path.DirectorySeparatorChar)
+				{
+					BaseDirectoryLen = Idx + 1;
+				}
+				else if(Text[Idx] == '?' || Text[Idx] == '*' || (Idx + 2 < Text.Length && Text[Idx] == '.' && Text[Idx + 1] == '.' && Text[Idx + 2] == '.'))
+				{
+					break;
+				}
+			}
+
+			// Extract the base directory
+			BaseDirectory = DirectoryReference.Combine(RootDirectory, Text.ToString(0, BaseDirectoryLen));
+
+			// Convert any directory wildcards ("...") into complete directory wildcards ("\\...\\"). We internally treat use "...\\" as the wildcard 
+			// token so we can correctly match zero directories. Patterns such as "foo...bar" should require at least one directory separator, so 
+			// should be converted to "foo*\\...\\*bar".
+			for(int Idx = BaseDirectoryLen; Idx < Text.Length; Idx++)
+			{
+				if(Text[Idx] == '.' && Text[Idx + 1] == '.' && Text[Idx + 2] == '.')
+				{
+					// Insert a directory separator before
+					if(Idx > BaseDirectoryLen && Text[Idx - 1] != Path.DirectorySeparatorChar)
+					{
+						Text.Insert(Idx++, '*');
+						Text.Insert(Idx++, Path.DirectorySeparatorChar);
+					}
+
+					// Skip past the ellipsis
+					Idx += 3;
+
+					// Insert a directory separator after
+					if(Idx == Text.Length || Text[Idx] != Path.DirectorySeparatorChar)
+					{
+						Text.Insert(Idx++, Path.DirectorySeparatorChar);
+						Text.Insert(Idx++, '*');
+					}
+				}
+			}
+
+			// Parse the tokens
+			int LastIdx = BaseDirectoryLen;
+			for(int Idx = BaseDirectoryLen; Idx < Text.Length; Idx++)
+			{
+				if(Text[Idx] == '?' || Text[Idx] == '*')
+				{
+					Tokens.Add(Text.ToString(LastIdx, Idx - LastIdx));
+					Tokens.Add(Text.ToString(Idx, 1));
+					LastIdx = Idx + 1;
+				}
+				else if(Idx - 3 >= BaseDirectoryLen && Text[Idx] == Path.DirectorySeparatorChar && Text[Idx - 1] == '.' && Text[Idx - 2] == '.' && Text[Idx - 3] == '.')
+				{
+					Tokens.Add(Text.ToString(LastIdx, Idx - 3 - LastIdx));
+					Tokens.Add(Text.ToString(Idx - 3, 4));
+					LastIdx = Idx + 1;
+				}
+			}
+			Tokens.Add(Text.ToString(LastIdx, Text.Length - LastIdx));
+		}
+
+		/// <summary>
+		/// A pattern without wildcards may match either a single file or directory based on context. This pattern resolves to the later as necessary, producing a new pattern.
+		/// </summary>
+		/// <returns>Pattern which matches a directory</returns>
+		public FilePattern AsDirectoryPattern()
+		{
+			if(ContainsWildcards())
+			{
+				return this;
+			}
+			else
+			{
+				StringBuilder Pattern = new StringBuilder();
+				foreach(string Token in Tokens)
+				{
+					Pattern.Append(Token);
+				}
+				if(Pattern.Length > 0)
+				{
+					Pattern.Append(Path.DirectorySeparatorChar);
+				}
+				Pattern.Append("...");
+
+				return new FilePattern(BaseDirectory, Pattern.ToString());
+			}
+		}
+
+		/// <summary>
+		/// For a pattern that does not contain wildcards, returns the single file location
+		/// </summary>
+		/// <returns>Location of the referenced file</returns>
+		public FileReference GetSingleFile()
+		{
+			if(Tokens.Count == 1)
+			{
+				return FileReference.Combine(BaseDirectory, Tokens[0]);
+			}
+			else
+			{
+				throw new InvalidOperationException("File pattern does not reference a single file");
+			}
+		}
+
+		/// <summary>
+		/// Checks whether this pattern is explicitly a directory, ie. is terminated with a directory separator
+		/// </summary>
+		/// <returns>True if the pattern is a directory</returns>
+		public bool EndsWithDirectorySeparator()
+		{
+			string LastToken = Tokens[Tokens.Count - 1];
+			return LastToken.Length > 0 && LastToken[LastToken.Length - 1] == Path.DirectorySeparatorChar;
+		}
+
+		/// <summary>
+		/// Determines whether the pattern contains wildcards
+		/// </summary>
+		/// <returns>True if the pattern contains wildcards, false otherwise.</returns>
+		public bool ContainsWildcards()
+		{
+			return Tokens.Count > 1;
+		}
+
+		/// <summary>
+		/// Tests whether a pattern is compatible with another pattern (that is, that the number and type of wildcards match)
+		/// </summary>
+		/// <param name="Other">Pattern to compare against</param>
+		/// <returns>Whether the patterns are compatible.</returns>
+		public bool IsCompatibleWith(FilePattern Other)
+		{
+			// Check there are the same number of tokens in each pattern
+			if(Tokens.Count != Other.Tokens.Count)
+			{
+				return false;
+			}
+
+			// Check all the wildcard tokens match
+			for(int Idx = 1; Idx < Tokens.Count; Idx += 2)
+			{
+				if(Tokens[Idx] != Other.Tokens[Idx])
+				{
+					return false;
+				}
+			}
+			return true;
+		}
+
+		/// <summary>
+		/// Converts this pattern to a C# regex format string, which matches paths relative to the base directory formatted with native directory separators
+		/// </summary>
+		/// <returns>The regex pattern</returns>
+		public string GetRegexPattern()
+		{
+			StringBuilder Pattern = new StringBuilder("^");
+			Pattern.Append(Regex.Escape(Tokens[0]));
+			for(int Idx = 1; Idx < Tokens.Count; Idx += 2)
+			{
+				// Append the wildcard expression
+				if(Tokens[Idx] == "?")
+				{
+					Pattern.Append("([^\\/])");
+				}
+				else if(Tokens[Idx] == "*")
+				{
+					Pattern.Append("([^\\/]*)");
+				}
+				else
+				{
+					Pattern.AppendFormat("((?:.+{0})?)", Regex.Escape(Path.DirectorySeparatorChar.ToString()));
+				}
+
+				// Append the next sequence of characters to match
+				Pattern.Append(Regex.Escape(Tokens[Idx + 1]));
+			}
+			Pattern.Append("$");
+			return Pattern.ToString();
+		}
+
+		/// <summary>
+		/// Creates a regex replacement pattern
+		/// </summary>
+		/// <returns>String representing the regex replacement pattern</returns>
+		public string GetRegexReplacementPattern()
+		{
+			StringBuilder Pattern = new StringBuilder();
+			for(int Idx = 0;;Idx += 2)
+			{
+				// Append the escaped replacement character
+				Pattern.Append(Tokens[Idx].Replace("$", "$$"));
+
+				// Check if we've reached the end of the string
+				if(Idx == Tokens.Count - 1)
+				{
+					break;
+				}
+
+				// Insert the capture
+				Pattern.AppendFormat("${0}", (Idx / 2) + 1);
+			}
+			return Pattern.ToString();
+		}
+
+		/// <summary>
+		/// Creates a file mapping between a set of source patterns and a target pattern. All patterns should have a matching order and number of wildcards.
+		/// </summary>
+		/// <param name="Files">Files to use for the mapping</param>
+		/// <param name="SourcePatterns">List of source patterns</param>
+		/// <param name="TargetPattern">Matching output pattern</param>
+		/// <param name="Filter">Filter to apply to source files</param>
+		/// <param name="TargetFileToSourceFile">Dictionary to receive a mapping from target file to source file. An exception is thrown if multiple source files map to one target file, or a source file is also used as a target file.</param>
+		public static Dictionary<FileReference, FileReference> CreateMapping(HashSet<FileReference> Files, ref FilePattern SourcePattern, ref FilePattern TargetPattern)
+		{
+			// If the source pattern ends in a directory separator, or a set of input files are specified and it doesn't contain wildcards, treat it as a full directory match
+			if(SourcePattern.EndsWithDirectorySeparator())
+			{
+				SourcePattern = new FilePattern(SourcePattern.BaseDirectory, String.Join("", SourcePattern.Tokens) + "...");
+			}
+			else if(Files != null)
+			{
+				SourcePattern = SourcePattern.AsDirectoryPattern();
+			}
+
+			// If we have multiple potential source files, but no wildcards in the output pattern, assume it's a directory and append the pattern from the source.
+			if(SourcePattern.ContainsWildcards() && !TargetPattern.ContainsWildcards())
+			{
+				StringBuilder NewPattern = new StringBuilder();
+				foreach(string Token in TargetPattern.Tokens)
+				{
+					NewPattern.Append(Token);
+				}
+				if(NewPattern.Length > 0 && NewPattern[NewPattern.Length - 1] != Path.DirectorySeparatorChar)
+				{
+					NewPattern.Append(Path.DirectorySeparatorChar);
+				}
+				foreach(string Token in SourcePattern.Tokens)
+				{
+					NewPattern.Append(Token);
+				}
+				TargetPattern = new FilePattern(TargetPattern.BaseDirectory, NewPattern.ToString());
+			}
+
+			// If the target pattern ends with a directory separator, treat it as a full directory match if it has wildcards, or a copy of the source pattern if not
+			if(TargetPattern.EndsWithDirectorySeparator())
+			{
+				TargetPattern = new FilePattern(TargetPattern.BaseDirectory, String.Join("", TargetPattern.Tokens) + "...");
+			}
+
+			// Handle the case where source and target pattern are both individual files
+			Dictionary<FileReference, FileReference> TargetFileToSourceFile = new Dictionary<FileReference, FileReference>();
+			if(SourcePattern.ContainsWildcards() || TargetPattern.ContainsWildcards())
+			{
+				// Check the two patterns are compatible
+				if(!SourcePattern.IsCompatibleWith(TargetPattern))
+				{
+					throw new FilePatternException("File patterns '{0}' and '{1}' do not have matching wildcards", SourcePattern, TargetPattern);
+				}
+
+				// Create a filter to match the source files
+				FileFilter Filter = new FileFilter(FileFilterType.Exclude);
+				Filter.Include(String.Join("", SourcePattern.Tokens));
+
+				// Apply it to the source directory
+				List<FileReference> SourceFiles;
+				if(Files == null)
+				{
+					SourceFiles = Filter.ApplyToDirectory(SourcePattern.BaseDirectory, true);
+				}
+				else
+				{
+					SourceFiles = CheckInputFiles(Files, SourcePattern.BaseDirectory);
+				}
+
+				// Map them onto output files
+				FileReference[] TargetFiles = new FileReference[SourceFiles.Count];
+
+				// Get the source and target regexes
+				string SourceRegex = SourcePattern.GetRegexPattern();
+				string TargetRegex = TargetPattern.GetRegexReplacementPattern();
+				for(int Idx = 0; Idx < SourceFiles.Count; Idx++)
+				{
+					string SourceRelativePath = SourceFiles[Idx].MakeRelativeTo(SourcePattern.BaseDirectory);
+					string TargetRelativePath = Regex.Replace(SourceRelativePath, SourceRegex, TargetRegex);
+					TargetFiles[Idx] = FileReference.Combine(TargetPattern.BaseDirectory, TargetRelativePath);
+				}
+
+				// Add them to the output map
+				for(int Idx = 0; Idx < TargetFiles.Length; Idx++)
+				{
+					FileReference ExistingSourceFile;
+					if(TargetFileToSourceFile.TryGetValue(TargetFiles[Idx], out ExistingSourceFile) && ExistingSourceFile != SourceFiles[Idx])
+					{
+						throw new FilePatternException("Output file '{0}' is mapped from '{1}' and '{2}'", TargetFiles[Idx], ExistingSourceFile, SourceFiles[Idx]);
+					}
+					TargetFileToSourceFile[TargetFiles[Idx]] = SourceFiles[Idx];
+				}
+			}
+			else
+			{
+				// Just copy a single file
+				FileReference SourceFile = SourcePattern.GetSingleFile();
+				if(FileReference.Exists(SourceFile))
+				{
+					FileReference TargetFile = TargetPattern.GetSingleFile();
+					TargetFileToSourceFile[TargetFile] = SourceFile;
+				}
+				else
+				{
+					throw new FilePatternException("Source file '{0}' does not exist", SourceFile);
+				}
+			}
+
+			// Check that no source file is also destination file
+			foreach(FileReference SourceFile in TargetFileToSourceFile.Values)
+			{
+				if(TargetFileToSourceFile.ContainsKey(SourceFile))
+				{
+					throw new FilePatternException("'{0}' is listed as a source and target file", SourceFile);
+				}
+			}
+
+			// Return the map
+			return TargetFileToSourceFile;
+		}
+
+		/// <summary>
+		/// Checks that the given input files all exist and are under the given base directory
+		/// </summary>
+		/// <param name="InputFiles">Input files to check</param>
+		/// <param name="BaseDirectory">Base directory for files</param>
+		/// <returns>List of valid files</returns>
+		public static List<FileReference> CheckInputFiles(IEnumerable<FileReference> InputFiles, DirectoryReference BaseDirectory)
+		{
+			List<FileReference> Files = new List<FileReference>();
+			foreach(FileReference InputFile in InputFiles)
+			{
+				if(!InputFile.IsUnderDirectory(BaseDirectory))
+				{
+					throw new FilePatternException("Source file '{0}' is not under '{1}'", InputFile, BaseDirectory);
+				}
+				else if(!FileReference.Exists(InputFile))
+				{
+					throw new FilePatternException("Source file '{0}' does not exist", InputFile);
+				}
+				else
+				{
+					Files.Add(InputFile);
+				}
+			}
+			return Files;
+		}
+
+		/// <summary>
+		/// Formats the pattern as a string
+		/// </summary>
+		/// <returns>The original representation of this pattern</returns>
+		public override string ToString()
+		{
+			return BaseDirectory.ToString() + Path.DirectorySeparatorChar + String.Join("", Tokens);
+		}
+	}
+}
