Forking DotNETUtilities into EpicGames.Common, to allow unification with Horde. Uses Core-style hierarchical naming, and excludes classes that have counterparts in NET Core. DotNETUtilities will be de

diff --git a/Engine/Source/Programs/DotNet/EpicGames.Common/FileReference.cs b/Engine/Source/Programs/DotNet/EpicGames.Common/FileReference.cs
new file mode 100644
--- /dev/null
+++ b/Engine/Source/Programs/DotNet/EpicGames.Common/FileReference.cs
@@ -0,0 +1,771 @@
+﻿// Copyright Epic Games, Inc. All Rights Reserved.
+
+using System;
+using System.Collections.Generic;
+using System.IO;
+using System.Linq;
+using System.Text;
+using System.Threading.Tasks;
+
+namespace EpicGames.Common
+{
+	/// <summary>
+	/// Representation of an absolute file path. Allows fast hashing and comparisons.
+	/// </summary>
+	[Serializable]
+	public class FileReference : FileSystemReference, IEquatable<FileReference>
+	{
+		/// <summary>
+		/// Dummy enum to allow invoking the constructor which takes a sanitized full path
+		/// </summary>
+		public enum Sanitize
+		{
+			None
+		}
+
+		/// <summary>
+		/// Default constructor.
+		/// </summary>
+		/// <param name="InPath">Path to this file</param>
+		public FileReference(string InPath)
+			: base(Path.GetFullPath(InPath))
+		{
+			if(FullName[FullName.Length - 1] == '\\' || FullName[FullName.Length - 1] == '/')
+			{
+				throw new ArgumentException("File names may not be terminated by a path separator character");
+			}
+		}
+
+		/// <summary>
+		/// Construct a FileReference from a FileInfo object.
+		/// </summary>
+		/// <param name="InInfo">Path to this file</param>
+		public FileReference(FileInfo InInfo)
+			: base(InInfo.FullName)
+		{
+		}
+
+		/// <summary>
+		/// Default constructor.
+		/// </summary>
+		/// <param name="InFullName">The full sanitized path</param>
+		/// <param name="InSanitize">Dummary argument to use the sanitized overload</param>
+		public FileReference(string InFullName, Sanitize InSanitize)
+			: base(InFullName)
+		{
+		}
+
+		/// <summary>
+		/// Create a FileReference from a string. If the string is null, returns a null FileReference.
+		/// </summary>
+		/// <param name="FileName">FileName for the string</param>
+		/// <returns>Returns a FileReference representing the given string, or null.</returns>
+		public static FileReference FromString(string FileName)
+		{
+			if(String.IsNullOrEmpty(FileName))
+			{
+				return null;
+			}
+			else
+			{
+				return new FileReference(FileName);
+			}
+		}
+
+		/// <summary>
+		/// Gets the file name without path information
+		/// </summary>
+		/// <returns>A string containing the file name</returns>
+		public string GetFileName()
+		{
+			return Path.GetFileName(FullName);
+		}
+
+		/// <summary>
+		/// Gets the file name without path information or an extension
+		/// </summary>
+		/// <returns>A string containing the file name without an extension</returns>
+		public string GetFileNameWithoutExtension()
+		{
+			return Path.GetFileNameWithoutExtension(FullName);
+		}
+
+		/// <summary>
+		/// Gets the file name without path or any extensions
+		/// </summary>
+		/// <returns>A string containing the file name without an extension</returns>
+		public string GetFileNameWithoutAnyExtensions()
+		{
+			int StartIdx = FullName.LastIndexOf(Path.DirectorySeparatorChar) + 1;
+
+			int EndIdx = FullName.IndexOf('.', StartIdx);
+			if (EndIdx < StartIdx)
+			{
+				return FullName.Substring(StartIdx);
+			}
+			else
+			{
+				return FullName.Substring(StartIdx, EndIdx - StartIdx);
+			}
+		}
+
+		/// <summary>
+		/// Gets the extension for this filename
+		/// </summary>
+		/// <returns>A string containing the extension of this filename</returns>
+		public string GetExtension()
+		{
+			return Path.GetExtension(FullName);
+		}
+
+		/// <summary>
+		/// Change the file's extension to something else
+		/// </summary>
+		/// <param name="Extension">The new extension</param>
+		/// <returns>A FileReference with the same path and name, but with the new extension</returns>
+		public FileReference ChangeExtension(string Extension)
+		{
+			string NewFullName = Path.ChangeExtension(FullName, Extension);
+			return new FileReference(NewFullName, Sanitize.None);
+		}
+
+		/// <summary>
+		/// Gets the directory containing this file
+		/// </summary>
+		/// <returns>A new directory object representing the directory containing this object</returns>
+		public DirectoryReference Directory
+		{
+			get
+			{
+				int ParentLength = FullName.LastIndexOf(Path.DirectorySeparatorChar);
+
+				if (ParentLength == 2 && FullName[1] == ':')
+				{
+					// windows root detected (C:)
+					ParentLength++;
+				}
+
+				if (ParentLength == 0 && FullName[0] == Path.DirectorySeparatorChar)
+				{
+					// nix style root (/) detected
+					ParentLength = 1;
+				}
+
+				return new DirectoryReference(FullName.Substring(0, ParentLength), DirectoryReference.Sanitize.None);
+			}
+		}
+
+		/// <summary>
+		/// Combine several fragments with a base directory, to form a new filename
+		/// </summary>
+		/// <param name="BaseDirectory">The base directory</param>
+		/// <param name="Fragments">Fragments to combine with the base directory</param>
+		/// <returns>The new file name</returns>
+		public static FileReference Combine(DirectoryReference BaseDirectory, params string[] Fragments)
+		{
+			string FullName = FileSystemReference.CombineStrings(BaseDirectory, Fragments);
+			return new FileReference(FullName, Sanitize.None);
+		}
+
+		/// <summary>
+		/// Append a string to the end of a filename
+		/// </summary>
+		/// <param name="A">The base file reference</param>
+		/// <param name="B">Suffix to be appended</param>
+		/// <returns>The new file reference</returns>
+		public static FileReference operator +(FileReference A, string B)
+		{
+			return new FileReference(A.FullName + B, Sanitize.None);
+		}
+
+		/// <summary>
+		/// Compares two filesystem object names for equality. Uses the canonical name representation, not the display name representation.
+		/// </summary>
+		/// <param name="A">First object to compare.</param>
+		/// <param name="B">Second object to compare.</param>
+		/// <returns>True if the names represent the same object, false otherwise</returns>
+		public static bool operator ==(FileReference A, FileReference B)
+		{
+			if ((object)A == null)
+			{
+				return (object)B == null;
+			}
+			else
+			{
+				return (object)B != null && A.FullName.Equals(B.FullName, Comparison);
+			}
+		}
+
+		/// <summary>
+		/// Compares two filesystem object names for inequality. Uses the canonical name representation, not the display name representation.
+		/// </summary>
+		/// <param name="A">First object to compare.</param>
+		/// <param name="B">Second object to compare.</param>
+		/// <returns>False if the names represent the same object, true otherwise</returns>
+		public static bool operator !=(FileReference A, FileReference B)
+		{
+			return !(A == B);
+		}
+
+		/// <summary>
+		/// Compares against another object for equality.
+		/// </summary>
+		/// <param name="Obj">other instance to compare.</param>
+		/// <returns>True if the names represent the same object, false otherwise</returns>
+		public override bool Equals(object Obj)
+		{
+			return (Obj is FileReference) && ((FileReference)Obj) == this;
+		}
+
+		/// <summary>
+		/// Compares against another object for equality.
+		/// </summary>
+		/// <param name="Obj">other instance to compare.</param>
+		/// <returns>True if the names represent the same object, false otherwise</returns>
+		public bool Equals(FileReference Obj)
+		{
+			return Obj == this;
+		}
+
+		/// <summary>
+		/// Returns a hash code for this object
+		/// </summary>
+		/// <returns></returns>
+		public override int GetHashCode()
+		{
+			return Comparer.GetHashCode(FullName);
+		}
+
+		/// <summary>
+		/// Helper function to create a remote file reference. Unlike normal FileReference objects, these aren't converted to a full path in the local filesystem, but are
+		/// left as they are passed in.
+		/// </summary>
+		/// <param name="AbsolutePath">The absolute path in the remote file system</param>
+		/// <returns>New file reference</returns>
+		public static FileReference MakeRemote(string AbsolutePath)
+		{
+			return new FileReference(AbsolutePath, Sanitize.None);
+		}
+
+		/// <summary>
+		/// Makes a file location writeable; 
+		/// </summary>
+		/// <param name="Location">Location of the file</param>
+		public static void MakeWriteable(FileReference Location)
+		{
+			if(Exists(Location))
+			{
+				FileAttributes Attributes = GetAttributes(Location);
+				if((Attributes & FileAttributes.ReadOnly) != 0)
+				{
+					SetAttributes(Location, Attributes & ~FileAttributes.ReadOnly);
+				}
+			}
+		}
+
+		/// <summary>
+		/// Finds the correct case to match the location of this file on disk. Uses the given case for parts of the path that do not exist.
+		/// </summary>
+		/// <param name="Location">The path to find the correct case for</param>
+		/// <returns>Location of the file with the correct case</returns>
+		public static FileReference FindCorrectCase(FileReference Location)
+		{
+			return new FileReference(FileUtils.FindCorrectCase(Location.ToFileInfo()));
+		}
+
+		/// <summary>
+		/// Constructs a FileInfo object from this reference
+		/// </summary>
+		/// <returns>New FileInfo object</returns>
+		public FileInfo ToFileInfo()
+		{
+			return new FileInfo(FullName);
+		}
+
+		#region System.IO.File methods
+
+		/// <summary>
+		/// Copies a file from one location to another
+		/// </summary>
+		/// <param name="SourceLocation">Location of the source file</param>
+		/// <param name="TargetLocation">Location of the target file</param>
+		public static void Copy(FileReference SourceLocation, FileReference TargetLocation)
+		{
+			File.Copy(SourceLocation.FullName, TargetLocation.FullName);
+		}
+
+		/// <summary>
+		/// Copies a file from one location to another
+		/// </summary>
+		/// <param name="SourceLocation">Location of the source file</param>
+		/// <param name="TargetLocation">Location of the target file</param>
+		/// <param name="bOverwrite">Whether to overwrite the file in the target location</param>
+		public static void Copy(FileReference SourceLocation, FileReference TargetLocation, bool bOverwrite)
+		{
+			File.Copy(SourceLocation.FullName, TargetLocation.FullName, bOverwrite);
+		}
+
+		/// <summary>
+		/// Deletes this file
+		/// </summary>
+		public static void Delete(FileReference Location)
+		{
+			File.Delete(Location.FullName);
+		}
+
+		/// <summary>
+		/// Determines whether the given filename exists
+		/// </summary>
+		/// <returns>True if it exists, false otherwise</returns>
+		public static bool Exists(FileReference Location)
+		{
+			return File.Exists(Location.FullName);
+		}
+
+		/// <summary>
+		/// Gets the attributes for a file
+		/// </summary>
+		/// <param name="Location">Location of the file</param>
+		/// <returns>Attributes for the file</returns>
+		public static FileAttributes GetAttributes(FileReference Location)
+		{
+			return File.GetAttributes(Location.FullName);
+		}
+
+		/// <summary>
+		/// Gets the time that the file was last written to
+		/// </summary>
+		/// <param name="Location">Location of the file</param>
+		/// <returns>Last write time, in local time</returns>
+		public static DateTime GetLastWriteTime(FileReference Location)
+		{
+			return File.GetLastWriteTime(Location.FullName);
+		}
+
+		/// <summary>
+		/// Gets the time that the file was last written to
+		/// </summary>
+		/// <param name="Location">Location of the file</param>
+		/// <returns>Last write time, in UTC time</returns>
+		public static DateTime GetLastWriteTimeUtc(FileReference Location)
+		{
+			return File.GetLastWriteTimeUtc(Location.FullName);
+		}
+
+		/// <summary>
+		/// Moves a file from one location to another
+		/// </summary>
+		/// <param name="SourceLocation">Location of the source file</param>
+		/// <param name="TargetLocation">Location of the target file</param>
+		public static void Move(FileReference SourceLocation, FileReference TargetLocation)
+		{
+			File.Move(SourceLocation.FullName, TargetLocation.FullName);
+		}
+
+		/// <summary>
+		/// Opens a FileStream on the specified path with read/write access
+		/// </summary>
+		/// <param name="Location">Location of the file</param>
+		/// <param name="Mode">Mode to use when opening the file</param>
+		/// <returns>New filestream for the given file</returns>
+		public static FileStream Open(FileReference Location, FileMode Mode)
+		{
+			return File.Open(Location.FullName, Mode);
+		}
+
+		/// <summary>
+		/// Opens a FileStream on the specified path
+		/// </summary>
+		/// <param name="Location">Location of the file</param>
+		/// <param name="Mode">Mode to use when opening the file</param>
+		/// <param name="Access">Sharing mode for the new file</param>
+		/// <returns>New filestream for the given file</returns>
+		public static FileStream Open(FileReference Location, FileMode Mode, FileAccess Access)
+		{
+			return File.Open(Location.FullName, Mode, Access);
+		}
+
+		/// <summary>
+		/// Opens a FileStream on the specified path
+		/// </summary>
+		/// <param name="Location">Location of the file</param>
+		/// <param name="Mode">Mode to use when opening the file</param>
+		/// <param name="Access">Access mode for the new file</param>
+		/// <param name="Share">Sharing mode for the open file</param>
+		/// <returns>New filestream for the given file</returns>
+		public static FileStream Open(FileReference Location, FileMode Mode, FileAccess Access, FileShare Share)
+		{
+			return File.Open(Location.FullName, Mode, Access, Share);
+		}
+
+		/// <summary>
+		/// Reads the contents of a file
+		/// </summary>
+		/// <param name="Location">Location of the file</param>
+		/// <returns>Byte array containing the contents of the file</returns>
+		public static byte[] ReadAllBytes(FileReference Location)
+		{
+			return File.ReadAllBytes(Location.FullName);
+		}
+
+		/// <summary>
+		/// Reads the contents of a file
+		/// </summary>
+		/// <param name="Location">Location of the file</param>
+		/// <returns>Contents of the file as a single string</returns>
+		public static string ReadAllText(FileReference Location)
+		{
+			return File.ReadAllText(Location.FullName);
+		}
+
+		/// <summary>
+		/// Reads the contents of a file
+		/// </summary>
+		/// <param name="Location">Location of the file</param>
+		/// <param name="Encoding">Encoding of the file</param>
+		/// <returns>Contents of the file as a single string</returns>
+		public static string ReadAllText(FileReference Location, Encoding Encoding)
+		{
+			return File.ReadAllText(Location.FullName, Encoding);
+		}
+
+		/// <summary>
+		/// Reads the contents of a file
+		/// </summary>
+		/// <param name="Location">Location of the file</param>
+		/// <returns>String array containing the contents of the file</returns>
+		public static string[] ReadAllLines(FileReference Location)
+		{
+			return File.ReadAllLines(Location.FullName);
+		}
+
+		/// <summary>
+		/// Reads the contents of a file
+		/// </summary>
+		/// <param name="Location">Location of the file</param>
+		/// <param name="Encoding">The encoding to use when parsing the file</param>
+		/// <returns>String array containing the contents of the file</returns>
+		public static string[] ReadAllLines(FileReference Location, Encoding Encoding)
+		{
+			return File.ReadAllLines(Location.FullName, Encoding);
+		}
+
+		/// <summary>
+		/// Sets the attributes for a file
+		/// </summary>
+		/// <param name="Location">Location of the file</param>
+		/// <param name="Attributes">New attributes for the file</param>
+		public static void SetAttributes(FileReference Location, FileAttributes Attributes)
+		{
+			File.SetAttributes(Location.FullName, Attributes);
+		}
+
+		/// <summary>
+		/// Sets the time that the file was last written to
+		/// </summary>
+		/// <param name="Location">Location of the file</param>
+		/// <param name="LastWriteTime">Last write time, in local time</param>
+		public static void SetLastWriteTime(FileReference Location, DateTime LastWriteTime)
+		{
+			File.SetLastWriteTime(Location.FullName, LastWriteTime);
+		}
+
+		/// <summary>
+		/// Sets the time that the file was last written to
+		/// </summary>
+		/// <param name="Location">Location of the file</param>
+		/// <param name="LastWriteTimeUtc">Last write time, in UTC time</param>
+		public static void SetLastWriteTimeUtc(FileReference Location, DateTime LastWriteTimeUtc)
+		{
+			File.SetLastWriteTimeUtc(Location.FullName, LastWriteTimeUtc);
+		}
+
+		/// <summary>
+		/// Sets the time that the file was last accessed.
+		/// </summary>
+		/// <param name="Location">Location of the file.</param>
+		/// <param name="LastWriteTime">Last access time, in local time.</param>
+		public static void SetLastAccessTime(FileReference Location, DateTime LastWriteTime)
+		{
+			File.SetLastWriteTime(Location.FullName, LastWriteTime);
+		}
+
+		/// <summary>
+		/// Sets the time that the file was last accessed.
+		/// </summary>
+		/// <param name="Location">Location of the file.</param>
+		/// <param name="LastWriteTime">Last access time, in UTC time.</param>
+		public static void SetLastAccessTimeUtc(FileReference Location, DateTime LastWriteTimeUtc)
+		{
+			File.SetLastWriteTimeUtc(Location.FullName, LastWriteTimeUtc);
+		}
+
+		/// <summary>
+		/// Writes the contents of a file
+		/// </summary>
+		/// <param name="Location">Location of the file</param>
+		/// <param name="Contents">Contents of the file</param>
+		public static void WriteAllBytes(FileReference Location, byte[] Contents)
+		{
+			File.WriteAllBytes(Location.FullName, Contents);
+		}
+
+		/// <summary>
+		/// Writes the data to the given file, if it's different from what's there already
+		/// </summary>
+		/// <param name="Location">Location of the file</param>
+		/// <param name="Contents">Contents of the file</param>
+		public static void WriteAllBytesIfDifferent(FileReference Location, byte[] Contents)
+		{
+			if(FileReference.Exists(Location))
+			{
+				byte[] CurrentContents = FileReference.ReadAllBytes(Location);
+				if(Contents.AsSpan().SequenceEqual(CurrentContents))
+				{
+					return;
+				}
+			}
+			WriteAllBytes(Location, Contents);
+		}
+
+		/// <summary>
+		/// Writes the contents of a file
+		/// </summary>
+		/// <param name="Location">Location of the file</param>
+		/// <param name="Contents">Contents of the file</param>
+		public static void WriteAllLines(FileReference Location, IEnumerable<string> Contents)
+		{
+			File.WriteAllLines(Location.FullName, Contents);
+		}
+
+		/// <summary>
+		/// Writes the contents of a file
+		/// </summary>
+		/// <param name="Location">Location of the file</param>
+		/// <param name="Contents">Contents of the file</param>
+		/// <param name="Encoding">The encoding to use when parsing the file</param>
+		public static void WriteAllLines(FileReference Location, IEnumerable<string> Contents, Encoding Encoding)
+		{
+			File.WriteAllLines(Location.FullName, Contents, Encoding);
+		}
+
+		/// <summary>
+		/// Writes the contents of a file
+		/// </summary>
+		/// <param name="Location">Location of the file</param>
+		/// <param name="Contents">Contents of the file</param>
+		public static void WriteAllLines(FileReference Location, string[] Contents)
+		{
+			File.WriteAllLines(Location.FullName, Contents);
+		}
+
+		/// <summary>
+		/// Writes the contents of a file
+		/// </summary>
+		/// <param name="Location">Location of the file</param>
+		/// <param name="Contents">Contents of the file</param>
+		/// <param name="Encoding">The encoding to use when parsing the file</param>
+		public static void WriteAllLines(FileReference Location, string[] Contents, Encoding Encoding)
+		{
+			File.WriteAllLines(Location.FullName, Contents, Encoding);
+		}
+
+		/// <summary>
+		/// Writes the contents of a file
+		/// </summary>
+		/// <param name="Location">Location of the file</param>
+		/// <param name="Contents">Contents of the file</param>
+		public static void WriteAllText(FileReference Location, string Contents)
+		{
+			File.WriteAllText(Location.FullName, Contents);
+		}
+
+		/// <summary>
+		/// Writes the contents of a file
+		/// </summary>
+		/// <param name="Location">Location of the file</param>
+		/// <param name="Contents">Contents of the file</param>
+		/// <param name="Encoding">The encoding to use when parsing the file</param>
+		public static void WriteAllText(FileReference Location, string Contents, Encoding Encoding)
+		{
+			File.WriteAllText(Location.FullName, Contents, Encoding);
+		}
+
+		/// <summary>
+		/// Appends the contents to a file
+		/// </summary>
+		/// <param name="Location">Location of the file</param>
+		/// <param name="Contents">Contents to append to the file</param>
+		public static void AppendAllLines(FileReference Location, IEnumerable<string> Contents)
+		{
+			File.AppendAllLines(Location.FullName, Contents);
+		}
+
+		/// <summary>
+		/// Appends the contents to a file
+		/// </summary>
+		/// <param name="Location">Location of the file</param>
+		/// <param name="Contents">Contents to append to the file</param>
+		/// <param name="Encoding">The encoding to use when parsing the file</param>
+		public static void AppendAllLines(FileReference Location, IEnumerable<string> Contents, Encoding Encoding)
+		{
+			File.AppendAllLines(Location.FullName, Contents, Encoding);
+		}
+
+		/// <summary>
+		/// Appends the contents to a file
+		/// </summary>
+		/// <param name="Location">Location of the file</param>
+		/// <param name="Contents">Contents to append to the file</param>
+		public static void AppendAllLines(FileReference Location, string[] Contents)
+		{
+			File.AppendAllLines(Location.FullName, Contents);
+		}
+
+		/// <summary>
+		/// Appends the contents to a file
+		/// </summary>
+		/// <param name="Location">Location of the file</param>
+		/// <param name="Contents">Contents to append to the file</param>
+		/// <param name="Encoding">The encoding to use when parsing the file</param>
+		public static void AppendAllLines(FileReference Location, string[] Contents, Encoding Encoding)
+		{
+			File.AppendAllLines(Location.FullName, Contents, Encoding);
+		}
+
+		/// <summary>
+		/// Appends the contents to a file
+		/// </summary>
+		/// <param name="Location">Location of the file</param>
+		/// <param name="Contents">Contents to append to the file</param>
+		public static void AppendAllText(FileReference Location, string Contents)
+		{
+			File.AppendAllText(Location.FullName, Contents);
+		}
+
+		/// <summary>
+		/// Appends the contents to a file
+		/// </summary>
+		/// <param name="Location">Location of the file</param>
+		/// <param name="Contents">Contents to append to the file</param>
+		/// <param name="Encoding">The encoding to use when parsing the file</param>
+		public static void AppendAllText(FileReference Location, string Contents, Encoding Encoding)
+		{
+			File.AppendAllText(Location.FullName, Contents, Encoding);
+		}
+
+		#endregion
+	}
+
+	/// <summary>
+	/// Extension methods for FileReference functionality
+	/// </summary>
+	public static class FileReferenceExtensionMethods
+	{
+		/// <summary>
+		/// Manually serialize a file reference to a binary stream.
+		/// </summary>
+		/// <param name="Writer">Binary writer to write to</param>
+		/// <param name="File">The file reference to write</param>
+		public static void Write(this BinaryWriter Writer, FileReference File)
+		{
+			Writer.Write((File == null) ? String.Empty : File.FullName);
+		}
+
+		/// <summary>
+		/// Serializes a file reference, using a lookup table to avoid serializing the same name more than once.
+		/// </summary>
+		/// <param name="Writer">The writer to save this reference to</param>
+		/// <param name="File">A file reference to output; may be null</param>
+		/// <param name="FileToUniqueId">A lookup table that caches previous files that have been output, and maps them to unique id's.</param>
+		public static void Write(this BinaryWriter Writer, FileReference File, Dictionary<FileReference, int> FileToUniqueId)
+		{
+			int UniqueId;
+			if (File == null)
+			{
+				Writer.Write(-1);
+			}
+			else if (FileToUniqueId.TryGetValue(File, out UniqueId))
+			{
+				Writer.Write(UniqueId);
+			}
+			else
+			{
+				Writer.Write(FileToUniqueId.Count);
+				Writer.Write(File);
+				FileToUniqueId.Add(File, FileToUniqueId.Count);
+			}
+		}
+
+		/// <summary>
+		/// Manually deserialize a file reference from a binary stream.
+		/// </summary>
+		/// <param name="Reader">Binary reader to read from</param>
+		/// <returns>New FileReference object</returns>
+		public static FileReference ReadFileReference(this BinaryReader Reader)
+		{
+			string FullName = Reader.ReadString();
+			return (FullName.Length == 0) ? null : new FileReference(FullName, FileReference.Sanitize.None);
+		}
+
+		/// <summary>
+		/// Deserializes a file reference, using a lookup table to avoid writing the same name more than once.
+		/// </summary>
+		/// <param name="Reader">The source to read from</param>
+		/// <param name="UniqueFiles">List of previously read file references. The index into this array is used in place of subsequent ocurrences of the file.</param>
+		/// <returns>The file reference that was read</returns>
+		public static FileReference ReadFileReference(this BinaryReader Reader, List<FileReference> UniqueFiles)
+		{
+			int UniqueId = Reader.ReadInt32();
+			if (UniqueId == -1)
+			{
+				return null;
+			}
+			else if (UniqueId < UniqueFiles.Count)
+			{
+				return UniqueFiles[UniqueId];
+			}
+			else
+			{
+				FileReference Result = Reader.ReadFileReference();
+				UniqueFiles.Add(Result);
+				return Result;
+			}
+		}
+
+		/// <summary>
+		/// Writes a FileReference to a binary archive
+		/// </summary>
+		/// <param name="Writer">The writer to output data to</param>
+		/// <param name="File">The file reference to write</param>
+		public static void WriteFileReference(this BinaryArchiveWriter Writer, FileReference File)
+		{
+			if(File == null)
+			{
+				Writer.WriteString(null);
+			}
+			else
+			{
+				Writer.WriteString(File.FullName);
+			}
+		}
+
+		/// <summary>
+		/// Reads a FileReference from a binary archive
+		/// </summary>
+		/// <param name="Reader">Reader to serialize data from</param>
+		/// <returns>New file reference instance</returns>
+		public static FileReference ReadFileReference(this BinaryArchiveReader Reader)
+		{
+			string FullName = Reader.ReadString();
+			if(FullName == null)
+			{
+				return null;
+			}
+			else
+			{
+				return new FileReference(FullName, FileReference.Sanitize.None);
+			}
+		}
+	}
+}
diff --git a/Engine/Source/Programs/DotNet/EpicGames.Common/FileSystemReference.cs b/Engine/Source/Programs/DotNet/EpicGames.Common/FileSystemReference.cs
new file mode 100644
--- /dev/null
+++ b/Engine/Source/Programs/DotNet/EpicGames.Common/FileSystemReference.cs
@@ -0,0 +1,331 @@
+﻿// Copyright Epic Games, Inc. All Rights Reserved.
+
+using System;
+using System.Collections.Generic;
+using System.IO;
+using System.Linq;
+using System.Text;
+using System.Threading.Tasks;
+
+namespace EpicGames.Common
+{
+	/// <summary>
+	/// Base class for file system objects (files or directories).
+	/// </summary>
+	[Serializable]
+	public abstract class FileSystemReference
+	{
+		/// <summary>
+		/// The path to this object. Stored as an absolute path, with O/S preferred separator characters, and no trailing slash for directories.
+		/// </summary>
+		public readonly string FullName;
+
+		/// <summary>
+		/// The comparer to use for file system references
+		/// </summary>
+		public static readonly StringComparer Comparer = StringComparer.OrdinalIgnoreCase;
+
+		/// <summary>
+		/// The comparison to use for file system references
+		/// </summary>
+		public static readonly StringComparison Comparison = StringComparison.OrdinalIgnoreCase;
+
+		/// <summary>
+		/// Direct constructor for a path
+		/// </summary>
+		protected FileSystemReference(string InFullName)
+		{
+			FullName = InFullName;
+		}
+
+		/// <summary>
+		/// Direct constructor for a path
+		/// </summary>
+		protected FileSystemReference(string InFullName, string InCanonicalName)
+		{
+			FullName = InFullName;
+		}
+
+		/// <summary>
+		/// Create a full path by concatenating multiple strings
+		/// </summary>
+		/// <returns></returns>
+		static protected string CombineStrings(DirectoryReference BaseDirectory, params string[] Fragments)
+		{
+			// Get the initial string to append to, and strip any root directory suffix from it
+			StringBuilder NewFullName = new StringBuilder(BaseDirectory.FullName);
+			if (NewFullName.Length > 0 && NewFullName[NewFullName.Length - 1] == Path.DirectorySeparatorChar)
+			{
+				NewFullName.Remove(NewFullName.Length - 1, 1);
+			}
+
+			// Scan through the fragments to append, appending them to a string and updating the base length as we go
+			foreach (string Fragment in Fragments)
+			{
+				// Check if this fragment is an absolute path
+				if ((Fragment.Length >= 2 && Fragment[1] == ':') || (Fragment.Length >= 1 && (Fragment[0] == '\\' || Fragment[0] == '/')))
+				{
+					// It is. Reset the new name to the full version of this path.
+					NewFullName.Clear();
+					NewFullName.Append(Path.GetFullPath(Fragment).TrimEnd(Path.DirectorySeparatorChar));
+				}
+				else
+				{
+					// Append all the parts of this fragment to the end of the existing path.
+					int StartIdx = 0;
+					while (StartIdx < Fragment.Length)
+					{
+						// Find the end of this fragment. We may have been passed multiple paths in the same string.
+						int EndIdx = StartIdx;
+						while (EndIdx < Fragment.Length && Fragment[EndIdx] != '\\' && Fragment[EndIdx] != '/')
+						{
+							EndIdx++;
+						}
+
+						// Ignore any empty sections, like leading or trailing slashes, and '.' directory references.
+						int Length = EndIdx - StartIdx;
+						if (Length == 0)
+						{
+							// Multiple directory separators in a row; illegal.
+							throw new ArgumentException(String.Format("Path fragment '{0}' contains invalid directory separators.", Fragment));
+						}
+						else if (Length == 2 && Fragment[StartIdx] == '.' && Fragment[StartIdx + 1] == '.')
+						{
+							// Remove the last directory name
+							for (int SeparatorIdx = NewFullName.Length - 1; SeparatorIdx >= 0; SeparatorIdx--)
+							{
+								if (NewFullName[SeparatorIdx] == Path.DirectorySeparatorChar)
+								{
+									NewFullName.Remove(SeparatorIdx, NewFullName.Length - SeparatorIdx);
+									break;
+								}
+							}
+						}
+						else if (Length != 1 || Fragment[StartIdx] != '.')
+						{
+							// Append this fragment
+							NewFullName.Append(Path.DirectorySeparatorChar);
+							NewFullName.Append(Fragment, StartIdx, Length);
+						}
+
+						// Move to the next part
+						StartIdx = EndIdx + 1;
+					}
+				}
+			}
+
+			// Append the directory separator
+			if (NewFullName.Length == 0 || (NewFullName.Length == 2 && NewFullName[1] == ':'))
+			{
+				NewFullName.Append(Path.DirectorySeparatorChar);
+			}
+
+			// Set the new path variables
+			return NewFullName.ToString();
+		}
+
+		/// <summary>
+		/// Checks whether this name has the given extension.
+		/// </summary>
+		/// <param name="Extension">The extension to check</param>
+		/// <returns>True if this name has the given extension, false otherwise</returns>
+		public bool HasExtension(string Extension)
+		{
+			if (Extension.Length > 0 && Extension[0] != '.')
+			{
+				return FullName.Length >= Extension.Length + 1 && FullName[FullName.Length - Extension.Length - 1] == '.' && FullName.EndsWith(Extension, Comparison);
+			}
+			else
+			{
+				return FullName.EndsWith(Extension, Comparison);
+			}
+		}
+
+		/// <summary>
+		/// Determines if the given object is at or under the given directory
+		/// </summary>
+		/// <param name="Other">Directory to check against</param>
+		/// <returns>True if this path is under the given directory</returns>
+		public bool IsUnderDirectory(DirectoryReference Other)
+		{
+			return FullName.StartsWith(Other.FullName, Comparison) && (FullName.Length == Other.FullName.Length || FullName[Other.FullName.Length] == Path.DirectorySeparatorChar || Other.IsRootDirectory());
+		}
+
+		/// <summary>
+		/// Searches the path fragments for the given name. Only complete fragments are considered a match.
+		/// </summary>
+		/// <param name="Name">Name to check for</param>
+		/// <param name="Offset">Offset within the string to start the search</param>
+		/// <returns>True if the given name is found within the path</returns>
+		public bool ContainsName(string Name, int Offset)
+		{
+			return ContainsName(Name, Offset, FullName.Length - Offset);
+		}
+
+		/// <summary>
+		/// Searches the path fragments for the given name. Only complete fragments are considered a match.
+		/// </summary>
+		/// <param name="Name">Name to check for</param>
+		/// <param name="Offset">Offset within the string to start the search</param>
+		/// <param name="Length">Length of the substring to search</param>
+		/// <returns>True if the given name is found within the path</returns>
+		public bool ContainsName(string Name, int Offset, int Length)
+		{
+			// Check the substring to search is at least long enough to contain a match
+			if(Length < Name.Length)
+			{
+				return false;
+			}
+
+			// Find each occurence of the name within the remaining string, then test whether it's surrounded by directory separators
+			int MatchIdx = Offset;
+			for(;;)
+			{
+				// Find the next occurrence
+				MatchIdx = FullName.IndexOf(Name, MatchIdx, Offset + Length - MatchIdx, Comparison);
+				if(MatchIdx == -1)
+				{
+					return false;
+				}
+
+				// Check if the substring is a directory
+				int MatchEndIdx = MatchIdx + Name.Length;
+				if(FullName[MatchIdx - 1] == Path.DirectorySeparatorChar && (MatchEndIdx == FullName.Length || FullName[MatchEndIdx] == Path.DirectorySeparatorChar))
+				{
+					return true;
+				}
+
+				// Move past the string that didn't match
+				MatchIdx += Name.Length;
+			}
+		}
+
+		/// <summary>
+		/// Determines if the given object is under the given directory, within a subfolder of the given name. Useful for masking out directories by name.
+		/// </summary>
+		/// <param name="Name">Name of a subfolder to also check for</param>
+		/// <param name="BaseDir">Base directory to check against</param>
+		/// <returns>True if the path is under the given directory</returns>
+		public bool ContainsName(string Name, DirectoryReference BaseDir)
+		{
+			// Check that this is under the base directory
+			if(!IsUnderDirectory(BaseDir))
+			{
+				return false;
+			}
+			else
+			{
+				return ContainsName(Name, BaseDir.FullName.Length);
+			}
+		}
+
+		/// <summary>
+		/// Determines if the given object is under the given directory, within a subfolder of the given name. Useful for masking out directories by name.
+		/// </summary>
+		/// <param name="Names">Names of subfolders to also check for</param>
+		/// <param name="BaseDir">Base directory to check against</param>
+		/// <returns>True if the path is under the given directory</returns>
+		public bool ContainsAnyNames(IEnumerable<string> Names, DirectoryReference BaseDir)
+		{
+			// Check that this is under the base directory
+			if(!IsUnderDirectory(BaseDir))
+			{
+				return false;
+			}
+			else
+			{
+				return Names.Any(x => ContainsName(x, BaseDir.FullName.Length));
+			}
+		}
+		
+		/// <summary>
+		/// Creates a relative path from the given base directory
+		/// </summary>
+		/// <param name="Directory">The directory to create a relative path from</param>
+		/// <returns>A relative path from the given directory</returns>
+		public string MakeRelativeTo(DirectoryReference Directory)
+		{
+			// Find how much of the path is common between the two paths. This length does not include a trailing directory separator character.
+			int CommonDirectoryLength = -1;
+			for (int Idx = 0; ; Idx++)
+			{
+				if (Idx == FullName.Length)
+				{
+					// The two paths are identical. Just return the "." character.
+					if (Idx == Directory.FullName.Length)
+					{
+						return ".";
+					}
+
+					// Check if we're finishing on a complete directory name
+					if (Directory.FullName[Idx] == Path.DirectorySeparatorChar)
+					{
+						CommonDirectoryLength = Idx;
+					}
+					break;
+				}
+				else if (Idx == Directory.FullName.Length)
+				{
+					// Check whether the end of the directory name coincides with a boundary for the current name.
+					if (FullName[Idx] == Path.DirectorySeparatorChar)
+					{
+						CommonDirectoryLength = Idx;
+					}
+					break;
+				}
+				else
+				{
+					// Check the two paths match, and bail if they don't. Increase the common directory length if we've reached a separator.
+					if(String.Compare(FullName, Idx, Directory.FullName, Idx, 1, Comparison) != 0)
+					{
+						break;
+					}
+					if (FullName[Idx] == Path.DirectorySeparatorChar)
+					{
+						CommonDirectoryLength = Idx;
+					}
+				}
+			}
+
+			// If there's no relative path, just return the absolute path
+			if (CommonDirectoryLength == -1)
+			{
+				return FullName;
+			}
+
+			// Append all the '..' separators to get back to the common directory, then the rest of the string to reach the target item
+			StringBuilder Result = new StringBuilder();
+			for (int Idx = CommonDirectoryLength + 1; Idx < Directory.FullName.Length; Idx++)
+			{
+				// Move up a directory
+				Result.Append("..");
+				Result.Append(Path.DirectorySeparatorChar);
+
+				// Scan to the next directory separator
+				while (Idx < Directory.FullName.Length && Directory.FullName[Idx] != Path.DirectorySeparatorChar)
+				{
+					Idx++;
+				}
+			}
+			if (CommonDirectoryLength + 1 < FullName.Length)
+			{
+				Result.Append(FullName, CommonDirectoryLength + 1, FullName.Length - CommonDirectoryLength - 1);
+			}
+			return Result.ToString();
+		}
+
+		public string ToNormalizedPath()
+		{
+			return FullName.Replace("\\", "/");
+		}
+
+		/// <summary>
+		/// Returns a string representation of this filesystem object
+		/// </summary>
+		/// <returns>Full path to the object</returns>
+		public override string ToString()
+		{
+			return FullName;
+		}
+	}
+}
diff --git a/Engine/Source/Programs/DotNet/EpicGames.Common/ManagedProcess.cs b/Engine/Source/Programs/DotNet/EpicGames.Common/ManagedProcess.cs
new file mode 100644
--- /dev/null
+++ b/Engine/Source/Programs/DotNet/EpicGames.Common/ManagedProcess.cs
@@ -0,0 +1,744 @@
+// Copyright Epic Games, Inc. All Rights Reserved.
+
+using Microsoft.Win32.SafeHandles;
+using System;
+using System.Collections.Generic;
+using System.ComponentModel;
+using System.Diagnostics;
+using System.IO;
+using System.Linq;
+using System.Runtime.InteropServices;
+using System.Text;
+using System.Threading;
+using System.Threading.Tasks;
+
+namespace EpicGames.Common
+{
+	/// <summary>
+	/// Tracks a set of processes, and destroys them when the object is disposed.
+	/// </summary>
+	public class ManagedProcessGroup : IDisposable
+	{
+		const UInt32 JOB_OBJECT_LIMIT_KILL_ON_JOB_CLOSE = 0x00002000;
+		const UInt32 JOB_OBJECT_LIMIT_BREAKAWAY_OK = 0x00000800;
+
+		[StructLayout(LayoutKind.Sequential)]
+		struct JOBOBJECT_BASIC_LIMIT_INFORMATION
+		{
+			public Int64 PerProcessUserTimeLimit;
+			public Int64 PerJobUserTimeLimit;
+			public UInt32 LimitFlags;
+			public UIntPtr MinimumWorkingSetSize;
+			public UIntPtr MaximumWorkingSetSize;
+			public UInt32 ActiveProcessLimit;
+			public Int64 Affinity;
+			public UInt32 PriorityClass;
+			public UInt32 SchedulingClass;
+		}
+
+		[StructLayout(LayoutKind.Sequential)]
+		struct IO_COUNTERS
+		{
+			public UInt64 ReadOperationCount;
+			public UInt64 WriteOperationCount;
+			public UInt64 OtherOperationCount;
+			public UInt64 ReadTransferCount;
+			public UInt64 WriteTransferCount;
+			public UInt64 OtherTransferCount;
+		}
+
+		[StructLayout(LayoutKind.Sequential)]
+		struct JOBOBJECT_EXTENDED_LIMIT_INFORMATION
+		{
+			public JOBOBJECT_BASIC_LIMIT_INFORMATION BasicLimitInformation;
+			public IO_COUNTERS IoInfo;
+			public UIntPtr ProcessMemoryLimit;
+			public UIntPtr JobMemoryLimit;
+			public UIntPtr PeakProcessMemoryUsed;
+			public UIntPtr PeakJobMemoryUsed;
+		}
+
+		[DllImport("kernel32.dll", SetLastError=true)]
+		static extern SafeFileHandle CreateJobObject(IntPtr SecurityAttributes, IntPtr Name);
+
+		const int JobObjectExtendedLimitInformation = 9;
+
+        [DllImport("kernel32.dll", SetLastError=true)]
+        static extern int SetInformationJobObject(SafeFileHandle hJob, int JobObjectInfoClass, IntPtr lpJobObjectInfo, int cbJobObjectInfoLength);
+
+		/// <summary>
+		/// Handle to the native job object that this process is added to. This handle is closed by the Dispose() method (and will automatically be closed by the OS on process exit),
+		/// resulting in the child process being killed.
+		/// </summary>
+		internal SafeFileHandle JobHandle
+		{
+			get;
+			private set;
+		}
+
+		/// <summary>
+		/// Determines support for using job objects
+		/// </summary>
+		static internal bool SupportsJobObjects
+		{
+			get { return Environment.OSVersion.Platform == PlatformID.Win32NT || Environment.OSVersion.Platform == PlatformID.Win32S || Environment.OSVersion.Platform == PlatformID.Win32Windows; }
+		}
+
+		/// <summary>
+		/// Constructor
+		/// </summary>
+		public ManagedProcessGroup()
+		{
+			if(SupportsJobObjects)
+			{
+				// Create the job object that the child process will be added to
+				JobHandle = CreateJobObject(IntPtr.Zero, IntPtr.Zero);
+				if(JobHandle == null)
+				{
+					throw new Win32Exception();
+				}
+
+				// Configure the job object to terminate the processes added to it when the handle is closed
+				JOBOBJECT_EXTENDED_LIMIT_INFORMATION LimitInformation = new JOBOBJECT_EXTENDED_LIMIT_INFORMATION();
+				LimitInformation.BasicLimitInformation.LimitFlags = JOB_OBJECT_LIMIT_KILL_ON_JOB_CLOSE | JOB_OBJECT_LIMIT_BREAKAWAY_OK;
+
+				int Length = Marshal.SizeOf(typeof(JOBOBJECT_EXTENDED_LIMIT_INFORMATION));
+				IntPtr LimitInformationPtr = Marshal.AllocHGlobal(Length);
+				Marshal.StructureToPtr(LimitInformation, LimitInformationPtr, false);
+
+				if(SetInformationJobObject(JobHandle, JobObjectExtendedLimitInformation, LimitInformationPtr, Length) == 0)
+				{
+					throw new Win32Exception();
+				}
+			}
+		}
+
+		public void Add(IntPtr ProcessHandle)
+		{
+		}
+
+		/// <summary>
+		/// Dispose of the process group
+		/// </summary>
+		public void Dispose()
+		{
+			if(JobHandle != null)
+			{
+				JobHandle.Dispose();
+				JobHandle = null;
+			}
+		}
+	}
+
+	/// <summary>
+	/// Encapsulates a managed child process, from which we can read the console output.
+	/// Uses job objects to ensure that the process will be terminated automatically by the O/S if the current process is terminated, and polls pipe reads to avoid blocking on unmanaged code
+	/// if the calling thread is terminated. Currently only implemented for Windows; makes heavy use of P/Invoke.
+	/// </summary>
+	public sealed class ManagedProcess : IDisposable
+	{
+		[StructLayout(LayoutKind.Sequential)]
+		class SECURITY_ATTRIBUTES
+		{
+			public int nLength;
+			public IntPtr lpSecurityDescriptor;
+			public int bInheritHandle;
+		}
+
+		[StructLayout(LayoutKind.Sequential)]
+		class PROCESS_INFORMATION
+		{
+			public IntPtr hProcess;
+			public IntPtr hThread;
+			public uint dwProcessId;
+			public uint dwThreadId;
+		}
+
+		[StructLayout(LayoutKind.Sequential)]
+		class STARTUPINFO
+		{
+			public int cb;
+			public string lpReserved;
+			public string lpDesktop;
+			public string lpTitle;
+			public uint dwX;
+			public uint dwY;
+			public uint dwXSize;
+			public uint dwYSize;
+			public uint dwXCountChars;
+			public uint dwYCountChars;
+			public uint dwFillAttribute;
+			public uint dwFlags;
+			public short wShowWindow;
+			public short cbReserved2;
+			public IntPtr lpReserved2;
+			public SafeHandle hStdInput;
+			public SafeHandle hStdOutput;
+			public SafeHandle hStdError;
+		}
+
+		[DllImport("kernel32.dll", SetLastError=true)]
+		static extern int AssignProcessToJobObject(SafeFileHandle hJob, IntPtr hProcess);
+
+		[DllImport("kernel32.dll", SetLastError=true)]
+		static extern int CloseHandle(IntPtr hObject);
+
+		[DllImport("kernel32.dll", SetLastError = true)]
+		static extern int CreatePipe(out SafeFileHandle hReadPipe, out SafeFileHandle hWritePipe, SECURITY_ATTRIBUTES lpPipeAttributes, uint nSize);
+
+		[DllImport("kernel32.dll", SetLastError = true)]
+		static extern int SetHandleInformation(SafeFileHandle hObject, int dwMask, int dwFlags);
+
+		const int HANDLE_FLAG_INHERIT = 1;
+
+        const int STARTF_USESTDHANDLES = 0x00000100;
+
+		[Flags]
+		enum ProcessCreationFlags : int
+		{
+			CREATE_NO_WINDOW = 0x08000000,
+			CREATE_SUSPENDED = 0x00000004,
+			NORMAL_PRIORITY_CLASS = 0x00000020,
+			IDLE_PRIORITY_CLASS = 0x00000040,
+			HIGH_PRIORITY_CLASS = 0x00000080,
+			REALTIME_PRIORITY_CLASS = 0x00000100,
+			BELOW_NORMAL_PRIORITY_CLASS = 0x00004000,
+			ABOVE_NORMAL_PRIORITY_CLASS = 0x00008000,
+		}
+
+        [DllImport("kernel32.dll", SetLastError=true)]
+        static extern int CreateProcess(/*[MarshalAs(UnmanagedType.LPTStr)]*/ string lpApplicationName, StringBuilder lpCommandLine, IntPtr lpProcessAttributes, IntPtr lpThreadAttributes, bool bInheritHandles, ProcessCreationFlags dwCreationFlags, IntPtr lpEnvironment, /*[MarshalAs(UnmanagedType.LPTStr)]*/ string lpCurrentDirectory, STARTUPINFO lpStartupInfo, PROCESS_INFORMATION lpProcessInformation);
+
+		[DllImport("kernel32.dll", SetLastError=true)]
+		static extern int ResumeThread(IntPtr hThread);
+
+		[DllImport("kernel32.dll", SetLastError=true)]
+		static extern IntPtr GetCurrentProcess();
+
+        const int DUPLICATE_SAME_ACCESS = 2;
+
+        [DllImport("kernel32.dll", SetLastError=true)]    
+        static extern int DuplicateHandle(
+            IntPtr hSourceProcessHandle,
+            SafeHandle hSourceHandle,
+            IntPtr hTargetProcess,
+            out SafeWaitHandle targetHandle,
+            int dwDesiredAccess,
+            [MarshalAs(UnmanagedType.Bool)] bool bInheritHandle,
+            int dwOptions
+        );
+
+		[DllImport("kernel32.dll", SetLastError = true)]
+		static extern int GetExitCodeProcess(SafeFileHandle hProcess, out int lpExitCode);
+
+		[DllImport ("kernel32.dll")]
+		static extern bool IsProcessInJob(IntPtr hProcess, IntPtr hJob, out bool Result);
+
+		[DllImport("kernel32.dll", SetLastError=true)]
+		static extern int TerminateProcess(SafeHandleZeroOrMinusOneIsInvalid hProcess, uint uExitCode);
+
+		const UInt32 INFINITE = 0xFFFFFFFF;
+
+		[DllImport("kernel32.dll", SetLastError=true)]
+		static extern UInt32 WaitForSingleObject(SafeHandleZeroOrMinusOneIsInvalid hHandle, UInt32 dwMilliseconds);
+
+		const int ERROR_ACCESS_DENIED = 5;
+
+		/// <summary>
+		/// Handle for the child process.
+		/// </summary>
+		SafeFileHandle ProcessHandle;
+
+		/// <summary>
+		/// The write end of the child process' stdin pipe.
+		/// </summary>
+		SafeFileHandle StdInWrite;
+
+		/// <summary>
+		/// The read end of the child process' stdout pipe.
+		/// </summary>
+		SafeFileHandle StdOutRead;
+
+		/// <summary>
+		/// Output stream for the child process.
+		/// </summary>
+		FileStream InnerStream;
+
+		/// <summary>
+		/// Reader for the process' output stream.
+		/// </summary>
+		StreamReader ReadStream;
+
+		/// <summary>
+		/// Standard process implementation for non-Windows platforms.
+		/// </summary>
+		Process FrameworkProcess;
+
+		/// <summary>
+		/// Static lock object. This is used to synchronize the creation of child processes - in particular, the inheritance of stdout/stderr write pipes. If processes
+		/// inherit pipes meant for other processes, they won't be closed until both terminate.
+		/// </summary>
+		static object LockObject = new object();
+
+		/// <summary>
+		/// Spawns a new managed process.
+		/// </summary>
+		/// <param name="Group">The managed process group to add to</param>
+		/// <param name="FileName">Path to the executable to be run</param>
+		/// <param name="CommandLine">Command line arguments for the process</param>
+		/// <param name="WorkingDirectory">Working directory for the new process. May be null to use the current working directory.</param>
+		/// <param name="Environment">Environment variables for the new process. May be null, in which case the current process' environment is inherited</param>
+		/// <param name="Input">Text to be passed via stdin to the new process. May be null.</param>
+		/// <param name="Priority">Priority for the child process</param>
+		public ManagedProcess(ManagedProcessGroup Group, string FileName, string CommandLine, string WorkingDirectory, IReadOnlyDictionary<string, string> Environment, byte[] Input, ProcessPriorityClass Priority)
+		{
+			// Create the child process
+			// NOTE: Child process must be created in a separate method to avoid stomping exception callstacks (https://stackoverflow.com/a/2494150)
+			try
+			{
+				if (ManagedProcessGroup.SupportsJobObjects)
+				{
+					CreateManagedProcessWin32(Group, FileName, CommandLine, WorkingDirectory, Environment, Input, Priority);
+				}
+				else
+				{
+					CreateManagedProcessPortable(Group, FileName, CommandLine, WorkingDirectory, Environment, Input, Priority);
+				}
+			}
+			catch (Exception Ex)
+			{
+				ExceptionUtils.AddContext(Ex, "while launching {0} {1}", FileName, CommandLine);
+				throw;
+			}
+		}
+
+		/// <summary>
+		/// Spawns a new managed process using Win32 native functions. 
+		/// </summary>
+		/// <param name="Group">The managed process group to add to</param>
+		/// <param name="FileName">Path to the executable to be run</param>
+		/// <param name="CommandLine">Command line arguments for the process</param>
+		/// <param name="WorkingDirectory">Working directory for the new process. May be null to use the current working directory.</param>
+		/// <param name="Environment">Environment variables for the new process. May be null, in which case the current process' environment is inherited</param>
+		/// <param name="Input">Text to be passed via stdin to the new process. May be null.</param>
+		/// <param name="Priority">Priority for the child process</param>
+		private void CreateManagedProcessWin32(ManagedProcessGroup Group, string FileName, string CommandLine, string WorkingDirectory, IReadOnlyDictionary<string, string> Environment, byte[] Input, ProcessPriorityClass Priority)
+		{
+			IntPtr EnvironmentBlock = IntPtr.Zero;
+			try
+			{
+				// Create the environment block for the child process, if necessary.
+				if (Environment != null)
+				{
+					// The native format for the environment block is a sequence of null terminated strings with a final null terminator.
+					List<byte> EnvironmentBytes = new List<byte>();
+					foreach (KeyValuePair<string, string> Pair in Environment)
+					{
+						EnvironmentBytes.AddRange(Console.OutputEncoding.GetBytes(Pair.Key));
+						EnvironmentBytes.Add((byte)'=');
+						EnvironmentBytes.AddRange(Console.OutputEncoding.GetBytes(Pair.Value));
+						EnvironmentBytes.Add((byte)0);
+					}
+					EnvironmentBytes.Add((byte)0);
+
+					// Allocate an unmanaged block of memory to store it.
+					EnvironmentBlock = Marshal.AllocHGlobal(EnvironmentBytes.Count);
+					Marshal.Copy(EnvironmentBytes.ToArray(), 0, EnvironmentBlock, EnvironmentBytes.Count);
+				}
+
+				PROCESS_INFORMATION ProcessInfo = new PROCESS_INFORMATION();
+				try
+				{
+					// Get the flags to create the new process
+					ProcessCreationFlags Flags = ProcessCreationFlags.CREATE_NO_WINDOW | ProcessCreationFlags.CREATE_SUSPENDED;
+					switch (Priority)
+					{
+						case ProcessPriorityClass.Normal:
+							Flags |= ProcessCreationFlags.NORMAL_PRIORITY_CLASS;
+							break;
+						case ProcessPriorityClass.Idle:
+							Flags |= ProcessCreationFlags.IDLE_PRIORITY_CLASS;
+							break;
+						case ProcessPriorityClass.High:
+							Flags |= ProcessCreationFlags.HIGH_PRIORITY_CLASS;
+							break;
+						case ProcessPriorityClass.RealTime:
+							Flags |= ProcessCreationFlags.REALTIME_PRIORITY_CLASS;
+							break;
+						case ProcessPriorityClass.BelowNormal:
+							Flags |= ProcessCreationFlags.BELOW_NORMAL_PRIORITY_CLASS;
+							break;
+						case ProcessPriorityClass.AboveNormal:
+							Flags |= ProcessCreationFlags.ABOVE_NORMAL_PRIORITY_CLASS;
+							break;
+					}
+
+					// Acquire a global lock before creating inheritable handles. If multiple threads create inheritable handles at the same time, and child processes will inherit them all.
+					// Since we need to wait for output pipes to be closed (in order to consume all output), this can result in output reads not returning until all processes with the same
+					// inherited handles are closed.
+					lock (LockObject)
+					{
+						SafeFileHandle StdInRead = null;
+						SafeFileHandle StdOutWrite = null;
+						SafeWaitHandle StdErrWrite = null;
+						try
+						{
+							// Create stdin and stdout pipes for the child process. We'll close the handles for the child process' ends after it's been created.
+							SECURITY_ATTRIBUTES SecurityAttributes = new SECURITY_ATTRIBUTES();
+							SecurityAttributes.nLength = Marshal.SizeOf(SecurityAttributes);
+							SecurityAttributes.bInheritHandle = 1;
+
+							if (CreatePipe(out StdInRead, out StdInWrite, SecurityAttributes, 0) == 0 || SetHandleInformation(StdInWrite, HANDLE_FLAG_INHERIT, 0) == 0)
+							{
+								throw new Win32ExceptionWithCode("Unable to create stdin pipe");
+							}
+							if (CreatePipe(out StdOutRead, out StdOutWrite, SecurityAttributes, 0) == 0 || SetHandleInformation(StdOutRead, HANDLE_FLAG_INHERIT, 0) == 0)
+							{
+								throw new Win32ExceptionWithCode("Unable to create stdout pipe");
+							}
+							if (DuplicateHandle(GetCurrentProcess(), StdOutWrite, GetCurrentProcess(), out StdErrWrite, 0, true, DUPLICATE_SAME_ACCESS) == 0)
+							{
+								throw new Win32ExceptionWithCode("Unable to duplicate stdout handle");
+							}
+
+							// Create the new process as suspended, so we can modify it before it starts executing (and potentially preempting us)
+							STARTUPINFO StartupInfo = new STARTUPINFO();
+							StartupInfo.cb = Marshal.SizeOf(StartupInfo);
+							StartupInfo.hStdInput = StdInRead;
+							StartupInfo.hStdOutput = StdOutWrite;
+							StartupInfo.hStdError = StdErrWrite;
+							StartupInfo.dwFlags = STARTF_USESTDHANDLES;
+
+							// Under heavy load (ie. spawning large number of processes, typically Clang) we see CreateProcess very occasionally failing with ERROR_ACCESS_DENIED.
+							int[] RetryDelay = { 100, 200, 1000, 5000 };
+							for(int AttemptIdx = 0; ; AttemptIdx++)
+							{
+								if (CreateProcess(null, new StringBuilder("\"" + FileName + "\" " + CommandLine), IntPtr.Zero, IntPtr.Zero, true, Flags, EnvironmentBlock, WorkingDirectory, StartupInfo, ProcessInfo) != 0)
+								{
+									break;
+								}
+
+								if (Marshal.GetLastWin32Error() != ERROR_ACCESS_DENIED || AttemptIdx >= RetryDelay.Length)
+								{
+									throw new Win32ExceptionWithCode("Unable to create process");
+								}
+
+								Log.TraceWarning("Unable to create process {0} (access denied); waiting {1}ms to retry...", FileName, RetryDelay[AttemptIdx]);
+								Thread.Sleep(RetryDelay[AttemptIdx]);
+							}
+						}
+						finally
+						{
+							// Close the write ends of the handle. We don't want any other process to be able to inherit these.
+							if (StdInRead != null)
+							{
+								StdInRead.Dispose();
+								StdInRead = null;
+							}
+							if (StdOutWrite != null)
+							{
+								StdOutWrite.Dispose();
+								StdOutWrite = null;
+							}
+							if (StdErrWrite != null)
+							{
+								StdErrWrite.Dispose();
+								StdErrWrite = null;
+							}
+						}
+					}
+
+					// Add it to our job object
+					if (Group != null && AssignProcessToJobObject(Group.JobHandle, ProcessInfo.hProcess) == 0)
+					{
+						// Support for nested job objects was only addeed in Windows 8; prior to that, assigning processes to job objects would fail. Figure out if we're already in a job, and ignore the error if we are.
+						int OriginalError = Marshal.GetLastWin32Error();
+
+						bool bProcessInJob;
+						IsProcessInJob(GetCurrentProcess(), IntPtr.Zero, out bProcessInJob);
+
+						if (!bProcessInJob)
+						{
+							throw new Win32ExceptionWithCode(OriginalError, "Unable to assign process to job object");
+						}
+					}
+
+					// Allow the thread to start running
+					if (ResumeThread(ProcessInfo.hThread) == -1)
+					{
+						throw new Win32ExceptionWithCode("Unable to resume thread in child process");
+					}
+
+					// If we have any input text, write it to stdin now
+					using (FileStream Stream = new FileStream(StdInWrite, FileAccess.Write, 4096, false))
+					{
+						if (Input != null)
+						{
+							Stream.Write(Input, 0, Input.Length);
+							Stream.Flush();
+						}
+					}
+
+					// Create the stream objects for reading the process output
+					InnerStream = new FileStream(StdOutRead, FileAccess.Read, 4096, false);
+					ReadStream = new StreamReader(InnerStream, Console.OutputEncoding);
+
+					// Wrap the process handle in a SafeFileHandle
+					ProcessHandle = new SafeFileHandle(ProcessInfo.hProcess, true);
+				}
+				finally
+				{
+					if (ProcessInfo.hProcess != IntPtr.Zero && ProcessHandle == null)
+					{
+						CloseHandle(ProcessInfo.hProcess);
+					}
+					if (ProcessInfo.hThread != IntPtr.Zero)
+					{
+						CloseHandle(ProcessInfo.hThread);
+					}
+				}
+			}
+			finally
+			{
+				if (EnvironmentBlock != IntPtr.Zero)
+				{
+					Marshal.FreeHGlobal(EnvironmentBlock);
+					EnvironmentBlock = IntPtr.Zero;
+				}
+			}
+		}
+
+		/// <summary>
+		/// Spawns a new managed process using Win32 native functions. 
+		/// </summary>
+		/// <param name="Group">The managed process group to add to</param>
+		/// <param name="FileName">Path to the executable to be run</param>
+		/// <param name="CommandLine">Command line arguments for the process</param>
+		/// <param name="WorkingDirectory">Working directory for the new process. May be null to use the current working directory.</param>
+		/// <param name="Environment">Environment variables for the new process. May be null, in which case the current process' environment is inherited</param>
+		/// <param name="Input">Text to be passed via stdin to the new process. May be null.</param>
+		/// <param name="Priority">Priority for the child process</param>
+		private void CreateManagedProcessPortable(ManagedProcessGroup Group, string FileName, string CommandLine, string WorkingDirectory, IReadOnlyDictionary<string, string> Environment, byte[] Input, ProcessPriorityClass Priority)
+		{
+			// Fallback for Mono platforms
+			FrameworkProcess = new Process();
+			FrameworkProcess.StartInfo.FileName = FileName;
+			FrameworkProcess.StartInfo.Arguments = CommandLine;
+			FrameworkProcess.StartInfo.WorkingDirectory = WorkingDirectory;
+			FrameworkProcess.StartInfo.RedirectStandardInput = true;
+			FrameworkProcess.StartInfo.RedirectStandardOutput = true;
+			FrameworkProcess.StartInfo.UseShellExecute = false;
+			FrameworkProcess.StartInfo.CreateNoWindow = true;
+
+			if (Environment != null)
+			{
+				foreach (KeyValuePair<string, string> Pair in Environment)
+				{
+					FrameworkProcess.StartInfo.EnvironmentVariables[Pair.Key] = Pair.Value;
+				}
+			}
+
+			FrameworkProcess.Start();
+			try
+			{
+				FrameworkProcess.PriorityClass = Priority;
+			}
+			catch
+			{
+
+			}
+			// If we have any input text, write it to stdin now
+			if (Input != null)
+			{
+				FrameworkProcess.StandardInput.BaseStream.Write(Input, 0, Input.Length);
+				FrameworkProcess.StandardInput.BaseStream.Close();
+			}
+		}
+
+		/// <summary>
+		/// Free the managed resources for this process
+		/// </summary>
+		public void Dispose()
+		{
+			if(ProcessHandle != null)
+			{
+				TerminateProcess(ProcessHandle, 0);
+				WaitForSingleObject(ProcessHandle, INFINITE);
+
+				ProcessHandle.Dispose();
+				ProcessHandle = null;
+			}
+			if(StdInWrite != null)
+			{
+				StdInWrite.Dispose();
+				StdInWrite = null;
+			}
+			if(StdOutRead != null)
+			{
+				StdOutRead.Dispose();
+				StdOutRead = null;
+			}
+			if(FrameworkProcess != null)
+			{
+				if(!FrameworkProcess.HasExited)
+				{
+					try
+					{
+						FrameworkProcess.Kill();
+						FrameworkProcess.WaitForExit();
+					}
+					catch
+					{
+					}
+				}
+				FrameworkProcess.Dispose();
+				FrameworkProcess = null;
+			}
+		}
+
+		/// <summary>
+		/// Reads data from the process output
+		/// </summary>
+		/// <param name="Buffer">The buffer to receive the data</param>
+		/// <param name="Offset">Offset within the buffer to write to</param>
+		/// <param name="Count">Maximum number of bytes to read</param>
+		/// <returns>Number of bytes read</returns>
+		public int Read(byte[] Buffer, int Offset, int Count)
+		{
+			// Fill the buffer, reentering managed code every 20ms to allow thread abort exceptions to be thrown
+			Task<int> ReadTask;
+			if(FrameworkProcess == null)
+			{
+				ReadTask = InnerStream.ReadAsync(Buffer, Offset, Count);
+			}
+			else
+			{
+				ReadTask = FrameworkProcess.StandardOutput.BaseStream.ReadAsync(Buffer, Offset, Count);
+			}
+			while(!ReadTask.Wait(20))
+			{
+				// Spin through managed code to allow things like ThreadAbortExceptions to be thrown.
+			}
+			return ReadTask.Result;
+		}
+
+		/// <summary>
+		/// Read all the output from the process. Does not return until the process terminates.
+		/// </summary>
+		/// <returns>List of output lines</returns>
+		public List<string> ReadAllLines()
+		{
+			// Manually read all the output lines from the stream. Using ReadToEndAsync() or ReadAsync() on the StreamReader is abysmally slow, especially
+			// for high-volume processes. Manually picking out the lines via a buffered ReadAsync() call was found to be 6x faster on 'p4 -ztag have' calls.
+			List<string> OutputLines = new List<string>();
+			byte[] Buffer = new byte[32 * 1024];
+			byte LastCharacter = 0;
+			int NumBytesInBuffer = 0;
+			for(;;)
+			{
+				// If we're got a single line larger than 32kb (!), enlarge the buffer to ensure we can handle it
+				if(NumBytesInBuffer == Buffer.Length)
+				{
+					Array.Resize(ref Buffer, Buffer.Length + 32 * 1024);
+				}
+
+				// Fill the buffer, reentering managed code every 20ms to allow thread abort exceptions to be thrown
+				int NumBytesRead = Read(Buffer, NumBytesInBuffer, Buffer.Length - NumBytesInBuffer);
+				if(NumBytesRead == 0)
+				{
+					if(NumBytesInBuffer > 0)
+					{
+						OutputLines.Add(Console.OutputEncoding.GetString(Buffer, 0, NumBytesInBuffer));
+					}
+					break;
+				}
+
+				// Otherwise append to the existing buffer
+				NumBytesInBuffer += NumBytesRead;
+
+				// Pull out all the complete output lines
+				int LastStartIdx = 0;
+				for(int Idx = 0; Idx < NumBytesInBuffer; Idx++)
+				{
+					if(Buffer[Idx] == '\r' || Buffer[Idx] == '\n')
+					{
+						if(Buffer[Idx] != '\n' || LastCharacter != '\r')
+						{
+							OutputLines.Add(Console.OutputEncoding.GetString(Buffer, LastStartIdx, Idx - LastStartIdx));
+						}
+						LastStartIdx = Idx + 1;
+					}
+					LastCharacter = Buffer[Idx];
+				}
+
+				// Shuffle everything back to the start of the buffer
+				Array.Copy(Buffer, LastStartIdx, Buffer, 0, Buffer.Length - LastStartIdx);
+				NumBytesInBuffer -= LastStartIdx;
+			}
+			return OutputLines;
+		}
+
+		/// <summary>
+		/// Block until the process outputs a line of text, or terminates.
+		/// </summary>
+		/// <param name="Line">Variable to receive the output line</param>
+		/// <param name="Token">Cancellation token which can be used to abort the read operation.</param>
+		/// <returns>True if a line was read, false if the process terminated without writing another line, or the cancellation token was signaled.</returns>
+		public bool TryReadLine(out string Line, CancellationToken? Token = null)
+		{
+			try
+			{
+				// Busy wait for the ReadLine call to finish, so we can get interrupted by thread abort exceptions
+				Task<string> ReadLineTask;
+				if(FrameworkProcess == null)
+				{
+					ReadLineTask = ReadStream.ReadLineAsync();
+				}
+				else
+				{
+					ReadLineTask = FrameworkProcess.StandardOutput.ReadLineAsync();
+				}
+				for (;;)
+				{
+					const int MillisecondsTimeout = 20;
+					if (Token.HasValue 
+							? ReadLineTask.Wait(MillisecondsTimeout, Token.Value) 
+							: ReadLineTask.Wait(MillisecondsTimeout))
+					{
+						Line = ReadLineTask.IsCompleted ? ReadLineTask.Result : null;
+						return Line != null;
+					}
+				}
+			}
+			catch (OperationCanceledException)
+			{
+				// If the cancel token is signalled, just return false.
+				Line = null;
+				return false;
+			}
+		}
+
+		/// <summary>
+		/// The exit code of the process. Throws an exception if the process has not terminated.
+		/// </summary>
+		public int ExitCode
+		{
+			get
+			{
+				if(FrameworkProcess == null)
+				{
+					int Value;
+					if(GetExitCodeProcess(ProcessHandle, out Value) == 0)
+					{
+						throw new Win32Exception();
+					}
+					return Value;
+				}
+				else
+				{
+					return FrameworkProcess.ExitCode;
+				}
+			}
+		}
+	}
+}
diff --git a/Engine/Source/Programs/DotNet/EpicGames.Perforce/EpicGames.Perforce.csproj b/Engine/Source/Programs/DotNet/EpicGames.Perforce/EpicGames.Perforce.csproj
--- a/Engine/Source/Programs/DotNet/EpicGames.Perforce/EpicGames.Perforce.csproj
+++ b/Engine/Source/Programs/DotNet/EpicGames.Perforce/EpicGames.Perforce.csproj
@@ -1,14 +1,12 @@
 <Project Sdk="Microsoft.NET.Sdk">
-
   <PropertyGroup>
-    <TargetFramework>netstandard2.0</TargetFramework>
+    <TargetFramework>netcoreapp3.1</TargetFramework>
     <Configuration Condition=" '$(Configuration)' == '' ">Development</Configuration>
     <OutputType>Library</OutputType>
     <RootNamespace>EpicGames.Common</RootNamespace>
     <GenerateAssemblyInfo>false</GenerateAssemblyInfo>
     <GenerateTargetFrameworkAttribute>false</GenerateTargetFrameworkAttribute>
     <Configurations>Debug;Release;Development</Configurations>
-    <Nullable>enable</Nullable>
   </PropertyGroup>
   <PropertyGroup Condition=" '$(Configuration)|$(Platform)' == 'Debug|AnyCPU' ">
     <DocumentationFile></DocumentationFile>
diff --git a/Engine/Source/Programs/DotNet/EpicGames.Perforce/PerforceChildProcess.cs b/Engine/Source/Programs/DotNet/EpicGames.Perforce/PerforceChildProcess.cs
--- a/Engine/Source/Programs/DotNet/EpicGames.Perforce/PerforceChildProcess.cs
+++ b/Engine/Source/Programs/DotNet/EpicGames.Perforce/PerforceChildProcess.cs
@@ -1,5 +1,6 @@
 // Copyright Epic Games, Inc. All Rights Reserved.
 
+using EpicGames.Common;
 using System;
 using System.Collections.Generic;
 using System.Diagnostics;
diff --git a/Engine/Source/Programs/UnrealBuildTool/ProjectFiles/ProjectFileGenerator.cs b/Engine/Source/Programs/UnrealBuildTool/ProjectFiles/ProjectFileGenerator.cs
--- a/Engine/Source/Programs/UnrealBuildTool/ProjectFiles/ProjectFileGenerator.cs
+++ b/Engine/Source/Programs/UnrealBuildTool/ProjectFiles/ProjectFileGenerator.cs
@@ -427,6 +427,36 @@ namespace UnrealBuildTool
 			}
 		}
 
+		/// <summary>
+		/// Adds all the DotNet folder to the solution.
+		/// </summary>
+		void AddSharedDotNetModules(MasterProjectFolder ProgramsFolder)
+		{
+			if (SupportsDotnetCoreProjects())
+			{
+				DirectoryInfo DotNetDir = new DirectoryInfo(DirectoryReference.Combine(UnrealBuildTool.EngineSourceDirectory, "Programs", "DotNet").FullName);
+				if (DotNetDir.Exists)
+				{
+					List<FileInfo> ProjectFiles = new List<FileInfo>();
+					foreach (DirectoryInfo ProjectDir in DotNetDir.EnumerateDirectories())
+					{
+						ProjectFiles.AddRange(ProjectDir.EnumerateFiles("*.csproj"));
+					}
+					if (ProjectFiles.Count > 0)
+					{
+						MasterProjectFolder Folder = ProgramsFolder.AddSubFolder("DotNet");
+						foreach (FileInfo ProjectFile in ProjectFiles)
+						{
+							VCSharpProjectFile Project = new VCSharpProjectFile(new FileReference(ProjectFile));
+							Project.ShouldBuildForAllSolutionTargets = false;
+							AddExistingProjectFile(Project, bForceDevelopmentConfiguration: true);
+							Folder.ChildProjects.Add(Project);
+						}
+					}
+				}
+			}
+		}
+
 		protected virtual bool SupportsDotnetCoreProjects()
 		{
 			return true;
@@ -915,6 +945,9 @@ namespace UnrealBuildTool
 					// Add automation.csproj files to the master project
 					AddAutomationModules(AllGameProjects, ProgramsFolder);
 
+					// Add shared projects
+					AddSharedDotNetModules(ProgramsFolder);
+
 					// Discover C# programs which should additionally be included in the solution.
 					DiscoverCSharpProgramProjects(ProgramsFolder);
                 }
@@ -1453,7 +1486,7 @@ namespace UnrealBuildTool
 						SubFolder.SubFolders.Count == 0 )
 					{
 						// 3)
-						if (SubFolder.FolderName.Equals(SubFolder.ChildProjects[0].ProjectFilePath.GetFileNameWithoutAnyExtensions(), StringComparison.InvariantCultureIgnoreCase))
+						if (SubFolder.FolderName.Equals(SubFolder.ChildProjects[0].ProjectFilePath.GetFileNameWithoutExtension(), StringComparison.InvariantCultureIgnoreCase))
 						{
 							CanCollapseFolder = true;
 						}
@@ -1520,7 +1553,7 @@ namespace UnrealBuildTool
 				{
 					if (CurChildProject != OtherChildProject)
 					{
-						if (CurChildProject.ProjectFilePath.GetFileNameWithoutAnyExtensions().Equals(OtherChildProject.ProjectFilePath.GetFileNameWithoutAnyExtensions(), StringComparison.InvariantCultureIgnoreCase))
+						if (CurChildProject.ProjectFilePath.GetFileName().Equals(OtherChildProject.ProjectFilePath.GetFileNameWithoutExtension(), StringComparison.InvariantCultureIgnoreCase))
 						{
 							throw new BuildException("Detected collision between two project files with the same path in the same master project folder, " + OtherChildProject.ProjectFilePath.FullName + " and " + CurChildProject.ProjectFilePath.FullName + " (master project folder: " + MasterProjectFolderPath + ")");
 						}
@@ -1535,7 +1568,7 @@ namespace UnrealBuildTool
 				// with project file names or file items, as that's not supported in Visual Studio.)
 				foreach (ProjectFile CurChildProject in Folder.ChildProjects)
 				{
-					if (CurChildProject.ProjectFilePath.GetFileNameWithoutAnyExtensions().Equals(SubFolder.FolderName, StringComparison.InvariantCultureIgnoreCase))
+					if (CurChildProject.ProjectFilePath.GetFileNameWithoutExtension().Equals(SubFolder.FolderName, StringComparison.InvariantCultureIgnoreCase))
 					{
 						throw new BuildException("Detected collision between a master project sub-folder " + SubFolder.FolderName + " and a project within the outer folder " + CurChildProject.ProjectFilePath + " (master project folder: " + MasterProjectFolderPath + ")");
 					}
