Rename Tools.DotNETCommon.Perforce namespace to EpicGames.Perforce.

#rnx

diff --git a/Engine/Source/Programs/DotNet/EpicGames.Perforce/PerforceAttributes.cs b/Engine/Source/Programs/DotNet/EpicGames.Perforce/PerforceAttributes.cs
--- a/Engine/Source/Programs/DotNet/EpicGames.Perforce/PerforceAttributes.cs
+++ b/Engine/Source/Programs/DotNet/EpicGames.Perforce/PerforceAttributes.cs
@@ -6,7 +6,7 @@ using System.Linq;
 using System.Text;
 using System.Threading.Tasks;
 
-namespace Tools.DotNETCommon.Perforce
+namespace EpicGames.Perforce
 {
 	/// <summary>
 	/// Attributes for fields that should be deserialized from P4 tags
diff --git a/Engine/Source/Programs/DotNet/EpicGames.Perforce/PerforceChildProcess.cs b/Engine/Source/Programs/DotNet/EpicGames.Perforce/PerforceChildProcess.cs
--- a/Engine/Source/Programs/DotNet/EpicGames.Perforce/PerforceChildProcess.cs
+++ b/Engine/Source/Programs/DotNet/EpicGames.Perforce/PerforceChildProcess.cs
@@ -9,7 +9,7 @@ using System.Linq;
 using System.Text;
 using System.Threading.Tasks;
 
-namespace Tools.DotNETCommon.Perforce
+namespace EpicGames.Perforce
 {
 	enum PerforceRecordType
 	{
diff --git a/Engine/Source/Programs/DotNet/EpicGames.Perforce/PerforceConnection.cs b/Engine/Source/Programs/DotNet/EpicGames.Perforce/PerforceConnection.cs
--- a/Engine/Source/Programs/DotNet/EpicGames.Perforce/PerforceConnection.cs
+++ b/Engine/Source/Programs/DotNet/EpicGames.Perforce/PerforceConnection.cs
@@ -10,7 +10,7 @@ using System.Reflection;
 using System.Text;
 using System.Threading.Tasks;
 
-namespace Tools.DotNETCommon.Perforce
+namespace EpicGames.Perforce
 {
 	/// <summary>
 	/// Stores settings for communicating with a Perforce server.
diff --git a/Engine/Source/Programs/DotNet/EpicGames.Perforce/PerforceError.cs b/Engine/Source/Programs/DotNet/EpicGames.Perforce/PerforceError.cs
--- a/Engine/Source/Programs/DotNet/EpicGames.Perforce/PerforceError.cs
+++ b/Engine/Source/Programs/DotNet/EpicGames.Perforce/PerforceError.cs
@@ -6,7 +6,7 @@ using System.Linq;
 using System.Text;
 using System.Threading.Tasks;
 
-namespace Tools.DotNETCommon.Perforce
+namespace EpicGames.Perforce
 {
 	#pragma warning disable CS1591
 
diff --git a/Engine/Source/Programs/DotNet/EpicGames.Perforce/PerforceException.cs b/Engine/Source/Programs/DotNet/EpicGames.Perforce/PerforceException.cs
--- a/Engine/Source/Programs/DotNet/EpicGames.Perforce/PerforceException.cs
+++ b/Engine/Source/Programs/DotNet/EpicGames.Perforce/PerforceException.cs
@@ -6,7 +6,7 @@ using System.Linq;
 using System.Text;
 using System.Threading.Tasks;
 
-namespace Tools.DotNETCommon.Perforce
+namespace EpicGames.Perforce
 {
 	/// <summary>
 	/// Represents an exception specific to perforce
diff --git a/Engine/Source/Programs/DotNet/EpicGames.Perforce/PerforceInfo.cs b/Engine/Source/Programs/DotNet/EpicGames.Perforce/PerforceInfo.cs
--- a/Engine/Source/Programs/DotNet/EpicGames.Perforce/PerforceInfo.cs
+++ b/Engine/Source/Programs/DotNet/EpicGames.Perforce/PerforceInfo.cs
@@ -6,7 +6,7 @@ using System.Linq;
 using System.Text;
 using System.Threading.Tasks;
 
-namespace Tools.DotNETCommon.Perforce
+namespace EpicGames.Perforce
 {
 	/// <summary>
 	/// Generic class for parsing "info" responses from Perforce
diff --git a/Engine/Source/Programs/DotNet/EpicGames.Perforce/PerforceResponse.cs b/Engine/Source/Programs/DotNet/EpicGames.Perforce/PerforceResponse.cs
--- a/Engine/Source/Programs/DotNet/EpicGames.Perforce/PerforceResponse.cs
+++ b/Engine/Source/Programs/DotNet/EpicGames.Perforce/PerforceResponse.cs
@@ -6,7 +6,7 @@ using System.Linq;
 using System.Text;
 using System.Threading.Tasks;
 
-namespace Tools.DotNETCommon.Perforce
+namespace EpicGames.Perforce
 {
 	/// <summary>
 	/// Base class for returning untyped response data
diff --git a/Engine/Source/Programs/DotNet/EpicGames.Perforce/PerforceResponseList.cs b/Engine/Source/Programs/DotNet/EpicGames.Perforce/PerforceResponseList.cs
--- a/Engine/Source/Programs/DotNet/EpicGames.Perforce/PerforceResponseList.cs
+++ b/Engine/Source/Programs/DotNet/EpicGames.Perforce/PerforceResponseList.cs
@@ -6,7 +6,7 @@ using System.Linq;
 using System.Text;
 using System.Threading.Tasks;
 
-namespace Tools.DotNETCommon.Perforce
+namespace EpicGames.Perforce
 {
 	/// <summary>
 	/// Represents a list of responses from the Perforce server. Within the list, individual responses
diff --git a/Engine/Source/Programs/DotNet/EpicGames.Perforce/PerforceUtils.cs b/Engine/Source/Programs/DotNet/EpicGames.Perforce/PerforceUtils.cs
--- a/Engine/Source/Programs/DotNet/EpicGames.Perforce/PerforceUtils.cs
+++ b/Engine/Source/Programs/DotNet/EpicGames.Perforce/PerforceUtils.cs
@@ -6,7 +6,7 @@ using System.Linq;
 using System.Text;
 using System.Threading.Tasks;
 
-namespace Tools.DotNETCommon.Perforce
+namespace EpicGames.Perforce
 {
 	public static class PerforceUtils
 	{
diff --git a/Engine/Source/Programs/DotNet/EpicGames.Perforce/Records/AddOptions.cs b/Engine/Source/Programs/DotNet/EpicGames.Perforce/Records/AddOptions.cs
--- a/Engine/Source/Programs/DotNet/EpicGames.Perforce/Records/AddOptions.cs
+++ b/Engine/Source/Programs/DotNet/EpicGames.Perforce/Records/AddOptions.cs
@@ -6,7 +6,7 @@ using System.Linq;
 using System.Text;
 using System.Threading.Tasks;
 
-namespace Tools.DotNETCommon.Perforce
+namespace EpicGames.Perforce
 {
 	/// <summary>
 	/// Options for the p4 add command
diff --git a/Engine/Source/Programs/DotNet/EpicGames.Perforce/Records/AddRecord.cs b/Engine/Source/Programs/DotNet/EpicGames.Perforce/Records/AddRecord.cs
--- a/Engine/Source/Programs/DotNet/EpicGames.Perforce/Records/AddRecord.cs
+++ b/Engine/Source/Programs/DotNet/EpicGames.Perforce/Records/AddRecord.cs
@@ -6,7 +6,7 @@ using System.Linq;
 using System.Text;
 using System.Threading.Tasks;
 
-namespace Tools.DotNETCommon.Perforce
+namespace EpicGames.Perforce
 {
 	/// <summary>
 	/// Contains information about an added file
diff --git a/Engine/Source/Programs/DotNet/EpicGames.Perforce/Records/ChangeRecord.cs b/Engine/Source/Programs/DotNet/EpicGames.Perforce/Records/ChangeRecord.cs
--- a/Engine/Source/Programs/DotNet/EpicGames.Perforce/Records/ChangeRecord.cs
+++ b/Engine/Source/Programs/DotNet/EpicGames.Perforce/Records/ChangeRecord.cs
@@ -6,7 +6,7 @@ using System.Linq;
 using System.Text;
 using System.Threading.Tasks;
 
-namespace Tools.DotNETCommon.Perforce
+namespace EpicGames.Perforce
 {
 	/// <summary>
 	/// Contains summary information about a changelist
diff --git a/Engine/Source/Programs/DotNet/EpicGames.Perforce/Records/ChangeStatus.cs b/Engine/Source/Programs/DotNet/EpicGames.Perforce/Records/ChangeStatus.cs
--- a/Engine/Source/Programs/DotNet/EpicGames.Perforce/Records/ChangeStatus.cs
+++ b/Engine/Source/Programs/DotNet/EpicGames.Perforce/Records/ChangeStatus.cs
@@ -6,7 +6,7 @@ using System.Linq;
 using System.Text;
 using System.Threading.Tasks;
 
-namespace Tools.DotNETCommon.Perforce
+namespace EpicGames.Perforce
 {
 	/// <summary>
 	/// Indicates the status of a changelist
diff --git a/Engine/Source/Programs/DotNet/EpicGames.Perforce/Records/ChangeType.cs b/Engine/Source/Programs/DotNet/EpicGames.Perforce/Records/ChangeType.cs
--- a/Engine/Source/Programs/DotNet/EpicGames.Perforce/Records/ChangeType.cs
+++ b/Engine/Source/Programs/DotNet/EpicGames.Perforce/Records/ChangeType.cs
@@ -6,7 +6,7 @@ using System.Linq;
 using System.Text;
 using System.Threading.Tasks;
 
-namespace Tools.DotNETCommon.Perforce
+namespace EpicGames.Perforce
 {
 	/// <summary>
 	/// The type of a changelist
diff --git a/Engine/Source/Programs/DotNet/EpicGames.Perforce/Records/ChangesOptions.cs b/Engine/Source/Programs/DotNet/EpicGames.Perforce/Records/ChangesOptions.cs
--- a/Engine/Source/Programs/DotNet/EpicGames.Perforce/Records/ChangesOptions.cs
+++ b/Engine/Source/Programs/DotNet/EpicGames.Perforce/Records/ChangesOptions.cs
@@ -6,7 +6,7 @@ using System.Linq;
 using System.Text;
 using System.Threading.Tasks;
 
-namespace Tools.DotNETCommon.Perforce
+namespace EpicGames.Perforce
 {
 	/// <summary>
 	/// Options for the 'changes' command
diff --git a/Engine/Source/Programs/DotNet/EpicGames.Perforce/Records/ChangesRecord.cs b/Engine/Source/Programs/DotNet/EpicGames.Perforce/Records/ChangesRecord.cs
--- a/Engine/Source/Programs/DotNet/EpicGames.Perforce/Records/ChangesRecord.cs
+++ b/Engine/Source/Programs/DotNet/EpicGames.Perforce/Records/ChangesRecord.cs
@@ -6,7 +6,7 @@ using System.Linq;
 using System.Text;
 using System.Threading.Tasks;
 
-namespace Tools.DotNETCommon.Perforce
+namespace EpicGames.Perforce
 {
 	/// <summary>
 	/// Contains summary information about a changelist
diff --git a/Engine/Source/Programs/DotNet/EpicGames.Perforce/Records/ClientLineEndings.cs b/Engine/Source/Programs/DotNet/EpicGames.Perforce/Records/ClientLineEndings.cs
--- a/Engine/Source/Programs/DotNet/EpicGames.Perforce/Records/ClientLineEndings.cs
+++ b/Engine/Source/Programs/DotNet/EpicGames.Perforce/Records/ClientLineEndings.cs
@@ -6,7 +6,7 @@ using System.Linq;
 using System.Text;
 using System.Threading.Tasks;
 
-namespace Tools.DotNETCommon.Perforce
+namespace EpicGames.Perforce
 {
 	/// <summary>
 	/// Specifies how to treat line endings on files synced to the client
diff --git a/Engine/Source/Programs/DotNet/EpicGames.Perforce/Records/ClientOptions.cs b/Engine/Source/Programs/DotNet/EpicGames.Perforce/Records/ClientOptions.cs
--- a/Engine/Source/Programs/DotNet/EpicGames.Perforce/Records/ClientOptions.cs
+++ b/Engine/Source/Programs/DotNet/EpicGames.Perforce/Records/ClientOptions.cs
@@ -6,7 +6,7 @@ using System.Linq;
 using System.Text;
 using System.Threading.Tasks;
 
-namespace Tools.DotNETCommon.Perforce
+namespace EpicGames.Perforce
 {
 	/// <summary>
 	/// Options for the client
diff --git a/Engine/Source/Programs/DotNet/EpicGames.Perforce/Records/ClientRecord.cs b/Engine/Source/Programs/DotNet/EpicGames.Perforce/Records/ClientRecord.cs
--- a/Engine/Source/Programs/DotNet/EpicGames.Perforce/Records/ClientRecord.cs
+++ b/Engine/Source/Programs/DotNet/EpicGames.Perforce/Records/ClientRecord.cs
@@ -6,7 +6,7 @@ using System.Linq;
 using System.Text;
 using System.Threading.Tasks;
 
-namespace Tools.DotNETCommon.Perforce
+namespace EpicGames.Perforce
 {
 	/// <summary>
 	/// Represents a Perforce clientspec
diff --git a/Engine/Source/Programs/DotNet/EpicGames.Perforce/Records/ClientSubmitOptions.cs b/Engine/Source/Programs/DotNet/EpicGames.Perforce/Records/ClientSubmitOptions.cs
--- a/Engine/Source/Programs/DotNet/EpicGames.Perforce/Records/ClientSubmitOptions.cs
+++ b/Engine/Source/Programs/DotNet/EpicGames.Perforce/Records/ClientSubmitOptions.cs
@@ -6,7 +6,7 @@ using System.Linq;
 using System.Text;
 using System.Threading.Tasks;
 
-namespace Tools.DotNETCommon.Perforce
+namespace EpicGames.Perforce
 {
 	/// <summary>
 	/// Client options for controlling the default behavior of p4 submit
diff --git a/Engine/Source/Programs/DotNet/EpicGames.Perforce/Records/ClientsOptions.cs b/Engine/Source/Programs/DotNet/EpicGames.Perforce/Records/ClientsOptions.cs
--- a/Engine/Source/Programs/DotNet/EpicGames.Perforce/Records/ClientsOptions.cs
+++ b/Engine/Source/Programs/DotNet/EpicGames.Perforce/Records/ClientsOptions.cs
@@ -6,7 +6,7 @@ using System.Linq;
 using System.Text;
 using System.Threading.Tasks;
 
-namespace Tools.DotNETCommon.Perforce
+namespace EpicGames.Perforce
 {
 	/// <summary>
 	/// Options for the p4 clients command
diff --git a/Engine/Source/Programs/DotNet/EpicGames.Perforce/Records/ClientsRecord.cs b/Engine/Source/Programs/DotNet/EpicGames.Perforce/Records/ClientsRecord.cs
--- a/Engine/Source/Programs/DotNet/EpicGames.Perforce/Records/ClientsRecord.cs
+++ b/Engine/Source/Programs/DotNet/EpicGames.Perforce/Records/ClientsRecord.cs
@@ -6,7 +6,7 @@ using System.Linq;
 using System.Text;
 using System.Threading.Tasks;
 
-namespace Tools.DotNETCommon.Perforce
+namespace EpicGames.Perforce
 {
 	/// <summary>
 	/// Record returned by the p4 clients command
diff --git a/Engine/Source/Programs/DotNet/EpicGames.Perforce/Records/DeleteChangeOptions.cs b/Engine/Source/Programs/DotNet/EpicGames.Perforce/Records/DeleteChangeOptions.cs
--- a/Engine/Source/Programs/DotNet/EpicGames.Perforce/Records/DeleteChangeOptions.cs
+++ b/Engine/Source/Programs/DotNet/EpicGames.Perforce/Records/DeleteChangeOptions.cs
@@ -6,7 +6,7 @@ using System.Linq;
 using System.Text;
 using System.Threading.Tasks;
 
-namespace Tools.DotNETCommon.Perforce
+namespace EpicGames.Perforce
 {
 	/// <summary>
 	/// Options for the DeleteChange command (p4 change -d)
diff --git a/Engine/Source/Programs/DotNet/EpicGames.Perforce/Records/DeleteClientOptions.cs b/Engine/Source/Programs/DotNet/EpicGames.Perforce/Records/DeleteClientOptions.cs
--- a/Engine/Source/Programs/DotNet/EpicGames.Perforce/Records/DeleteClientOptions.cs
+++ b/Engine/Source/Programs/DotNet/EpicGames.Perforce/Records/DeleteClientOptions.cs
@@ -6,7 +6,7 @@ using System.Linq;
 using System.Text;
 using System.Threading.Tasks;
 
-namespace Tools.DotNETCommon.Perforce
+namespace EpicGames.Perforce
 {
 	/// <summary>
 	/// Options for the p4 client -d command
diff --git a/Engine/Source/Programs/DotNet/EpicGames.Perforce/Records/DeleteOptions.cs b/Engine/Source/Programs/DotNet/EpicGames.Perforce/Records/DeleteOptions.cs
--- a/Engine/Source/Programs/DotNet/EpicGames.Perforce/Records/DeleteOptions.cs
+++ b/Engine/Source/Programs/DotNet/EpicGames.Perforce/Records/DeleteOptions.cs
@@ -6,7 +6,7 @@ using System.Linq;
 using System.Text;
 using System.Threading.Tasks;
 
-namespace Tools.DotNETCommon.Perforce
+namespace EpicGames.Perforce
 {
 	/// <summary>
 	/// Options for the p4 delete command
diff --git a/Engine/Source/Programs/DotNet/EpicGames.Perforce/Records/DeleteRecord.cs b/Engine/Source/Programs/DotNet/EpicGames.Perforce/Records/DeleteRecord.cs
--- a/Engine/Source/Programs/DotNet/EpicGames.Perforce/Records/DeleteRecord.cs
+++ b/Engine/Source/Programs/DotNet/EpicGames.Perforce/Records/DeleteRecord.cs
@@ -6,7 +6,7 @@ using System.Linq;
 using System.Text;
 using System.Threading.Tasks;
 
-namespace Tools.DotNETCommon.Perforce
+namespace EpicGames.Perforce
 {
 	/// <summary>
 	/// Information about a file opened for delete
diff --git a/Engine/Source/Programs/DotNet/EpicGames.Perforce/Records/DescribeFileRecord.cs b/Engine/Source/Programs/DotNet/EpicGames.Perforce/Records/DescribeFileRecord.cs
--- a/Engine/Source/Programs/DotNet/EpicGames.Perforce/Records/DescribeFileRecord.cs
+++ b/Engine/Source/Programs/DotNet/EpicGames.Perforce/Records/DescribeFileRecord.cs
@@ -6,7 +6,7 @@ using System.Linq;
 using System.Text;
 using System.Threading.Tasks;
 
-namespace Tools.DotNETCommon.Perforce
+namespace EpicGames.Perforce
 {
 	/// <summary>
 	/// Contains information about an individual file in a returned DescribeRecord
diff --git a/Engine/Source/Programs/DotNet/EpicGames.Perforce/Records/DescribeOptions.cs b/Engine/Source/Programs/DotNet/EpicGames.Perforce/Records/DescribeOptions.cs
--- a/Engine/Source/Programs/DotNet/EpicGames.Perforce/Records/DescribeOptions.cs
+++ b/Engine/Source/Programs/DotNet/EpicGames.Perforce/Records/DescribeOptions.cs
@@ -6,7 +6,7 @@ using System.Linq;
 using System.Text;
 using System.Threading.Tasks;
 
-namespace Tools.DotNETCommon.Perforce
+namespace EpicGames.Perforce
 {
 	/// <summary>
 	/// Options for the p4 describe command
diff --git a/Engine/Source/Programs/DotNet/EpicGames.Perforce/Records/DescribeRecord.cs b/Engine/Source/Programs/DotNet/EpicGames.Perforce/Records/DescribeRecord.cs
--- a/Engine/Source/Programs/DotNet/EpicGames.Perforce/Records/DescribeRecord.cs
+++ b/Engine/Source/Programs/DotNet/EpicGames.Perforce/Records/DescribeRecord.cs
@@ -6,7 +6,7 @@ using System.Linq;
 using System.Text;
 using System.Threading.Tasks;
 
-namespace Tools.DotNETCommon.Perforce
+namespace EpicGames.Perforce
 {
 	/// <summary>
 	/// Detailed description of an individual changelist
diff --git a/Engine/Source/Programs/DotNet/EpicGames.Perforce/Records/EditOptions.cs b/Engine/Source/Programs/DotNet/EpicGames.Perforce/Records/EditOptions.cs
--- a/Engine/Source/Programs/DotNet/EpicGames.Perforce/Records/EditOptions.cs
+++ b/Engine/Source/Programs/DotNet/EpicGames.Perforce/Records/EditOptions.cs
@@ -6,7 +6,7 @@ using System.Linq;
 using System.Text;
 using System.Threading.Tasks;
 
-namespace Tools.DotNETCommon.Perforce
+namespace EpicGames.Perforce
 {
 	/// <summary>
 	/// Options for the p4 edit command
diff --git a/Engine/Source/Programs/DotNet/EpicGames.Perforce/Records/EditRecord.cs b/Engine/Source/Programs/DotNet/EpicGames.Perforce/Records/EditRecord.cs
--- a/Engine/Source/Programs/DotNet/EpicGames.Perforce/Records/EditRecord.cs
+++ b/Engine/Source/Programs/DotNet/EpicGames.Perforce/Records/EditRecord.cs
@@ -6,7 +6,7 @@ using System.Linq;
 using System.Text;
 using System.Threading.Tasks;
 
-namespace Tools.DotNETCommon.Perforce
+namespace EpicGames.Perforce
 {
 	/// <summary>
 	/// Information about a file opened for edit
diff --git a/Engine/Source/Programs/DotNet/EpicGames.Perforce/Records/FStatOptions.cs b/Engine/Source/Programs/DotNet/EpicGames.Perforce/Records/FStatOptions.cs
--- a/Engine/Source/Programs/DotNet/EpicGames.Perforce/Records/FStatOptions.cs
+++ b/Engine/Source/Programs/DotNet/EpicGames.Perforce/Records/FStatOptions.cs
@@ -6,7 +6,7 @@ using System.Linq;
 using System.Text;
 using System.Threading.Tasks;
 
-namespace Tools.DotNETCommon.Perforce
+namespace EpicGames.Perforce
 {
 	/// <summary>
 	/// Options for the Fstat command
diff --git a/Engine/Source/Programs/DotNet/EpicGames.Perforce/Records/FStatRecord.cs b/Engine/Source/Programs/DotNet/EpicGames.Perforce/Records/FStatRecord.cs
--- a/Engine/Source/Programs/DotNet/EpicGames.Perforce/Records/FStatRecord.cs
+++ b/Engine/Source/Programs/DotNet/EpicGames.Perforce/Records/FStatRecord.cs
@@ -6,7 +6,7 @@ using System.Linq;
 using System.Text;
 using System.Threading.Tasks;
 
-namespace Tools.DotNETCommon.Perforce
+namespace EpicGames.Perforce
 {
 	/// <summary>
 	/// Record returned by the fstat command
diff --git a/Engine/Source/Programs/DotNet/EpicGames.Perforce/Records/FileAction.cs b/Engine/Source/Programs/DotNet/EpicGames.Perforce/Records/FileAction.cs
--- a/Engine/Source/Programs/DotNet/EpicGames.Perforce/Records/FileAction.cs
+++ b/Engine/Source/Programs/DotNet/EpicGames.Perforce/Records/FileAction.cs
@@ -6,7 +6,7 @@ using System.Linq;
 using System.Text;
 using System.Threading.Tasks;
 
-namespace Tools.DotNETCommon.Perforce
+namespace EpicGames.Perforce
 {
 	/// <summary>
 	/// Actions applied to a file in a particular revision
diff --git a/Engine/Source/Programs/DotNet/EpicGames.Perforce/Records/FileLogOptions.cs b/Engine/Source/Programs/DotNet/EpicGames.Perforce/Records/FileLogOptions.cs
--- a/Engine/Source/Programs/DotNet/EpicGames.Perforce/Records/FileLogOptions.cs
+++ b/Engine/Source/Programs/DotNet/EpicGames.Perforce/Records/FileLogOptions.cs
@@ -6,7 +6,7 @@ using System.Linq;
 using System.Text;
 using System.Threading.Tasks;
 
-namespace Tools.DotNETCommon.Perforce
+namespace EpicGames.Perforce
 {
 	/// <summary>
 	/// Options for the filelog command
diff --git a/Engine/Source/Programs/DotNet/EpicGames.Perforce/Records/FileLogRecord.cs b/Engine/Source/Programs/DotNet/EpicGames.Perforce/Records/FileLogRecord.cs
--- a/Engine/Source/Programs/DotNet/EpicGames.Perforce/Records/FileLogRecord.cs
+++ b/Engine/Source/Programs/DotNet/EpicGames.Perforce/Records/FileLogRecord.cs
@@ -6,7 +6,7 @@ using System.Linq;
 using System.Text;
 using System.Threading.Tasks;
 
-namespace Tools.DotNETCommon.Perforce
+namespace EpicGames.Perforce
 {
 	/// <summary>
 	/// Record output by the filelog command
diff --git a/Engine/Source/Programs/DotNet/EpicGames.Perforce/Records/GetChangeOptions.cs b/Engine/Source/Programs/DotNet/EpicGames.Perforce/Records/GetChangeOptions.cs
--- a/Engine/Source/Programs/DotNet/EpicGames.Perforce/Records/GetChangeOptions.cs
+++ b/Engine/Source/Programs/DotNet/EpicGames.Perforce/Records/GetChangeOptions.cs
@@ -6,7 +6,7 @@ using System.Linq;
 using System.Text;
 using System.Threading.Tasks;
 
-namespace Tools.DotNETCommon.Perforce
+namespace EpicGames.Perforce
 {
 	/// <summary>
 	/// Options for the GetChange command (p4 change -o)
diff --git a/Engine/Source/Programs/DotNet/EpicGames.Perforce/Records/InfoOptions.cs b/Engine/Source/Programs/DotNet/EpicGames.Perforce/Records/InfoOptions.cs
--- a/Engine/Source/Programs/DotNet/EpicGames.Perforce/Records/InfoOptions.cs
+++ b/Engine/Source/Programs/DotNet/EpicGames.Perforce/Records/InfoOptions.cs
@@ -6,7 +6,7 @@ using System.Linq;
 using System.Text;
 using System.Threading.Tasks;
 
-namespace Tools.DotNETCommon.Perforce
+namespace EpicGames.Perforce
 {
 	/// <summary>
 	/// Options for the p4 info command
diff --git a/Engine/Source/Programs/DotNet/EpicGames.Perforce/Records/InfoRecord.cs b/Engine/Source/Programs/DotNet/EpicGames.Perforce/Records/InfoRecord.cs
--- a/Engine/Source/Programs/DotNet/EpicGames.Perforce/Records/InfoRecord.cs
+++ b/Engine/Source/Programs/DotNet/EpicGames.Perforce/Records/InfoRecord.cs
@@ -6,7 +6,7 @@ using System.Linq;
 using System.Text;
 using System.Threading.Tasks;
 
-namespace Tools.DotNETCommon.Perforce
+namespace EpicGames.Perforce
 {
 	/// <summary>
 	/// Contains information about the current client and server configuration
diff --git a/Engine/Source/Programs/DotNet/EpicGames.Perforce/Records/IntegrationAction.cs b/Engine/Source/Programs/DotNet/EpicGames.Perforce/Records/IntegrationAction.cs
--- a/Engine/Source/Programs/DotNet/EpicGames.Perforce/Records/IntegrationAction.cs
+++ b/Engine/Source/Programs/DotNet/EpicGames.Perforce/Records/IntegrationAction.cs
@@ -6,7 +6,7 @@ using System.Linq;
 using System.Text;
 using System.Threading.Tasks;
 
-namespace Tools.DotNETCommon.Perforce
+namespace EpicGames.Perforce
 {
 	/// <summary>
 	/// Describes the action performed by the user when resolving the integration
diff --git a/Engine/Source/Programs/DotNet/EpicGames.Perforce/Records/IntegrationRecord.cs b/Engine/Source/Programs/DotNet/EpicGames.Perforce/Records/IntegrationRecord.cs
--- a/Engine/Source/Programs/DotNet/EpicGames.Perforce/Records/IntegrationRecord.cs
+++ b/Engine/Source/Programs/DotNet/EpicGames.Perforce/Records/IntegrationRecord.cs
@@ -6,7 +6,7 @@ using System.Linq;
 using System.Text;
 using System.Threading.Tasks;
 
-namespace Tools.DotNETCommon.Perforce
+namespace EpicGames.Perforce
 {
 	/// <summary>
 	/// Stores integration information for a file revision
diff --git a/Engine/Source/Programs/DotNet/EpicGames.Perforce/Records/OpenedOptions.cs b/Engine/Source/Programs/DotNet/EpicGames.Perforce/Records/OpenedOptions.cs
--- a/Engine/Source/Programs/DotNet/EpicGames.Perforce/Records/OpenedOptions.cs
+++ b/Engine/Source/Programs/DotNet/EpicGames.Perforce/Records/OpenedOptions.cs
@@ -6,7 +6,7 @@ using System.Linq;
 using System.Text;
 using System.Threading.Tasks;
 
-namespace Tools.DotNETCommon.Perforce
+namespace EpicGames.Perforce
 {
 	/// <summary>
 	/// Options for the p4 opened command
diff --git a/Engine/Source/Programs/DotNet/EpicGames.Perforce/Records/PrintRecord.cs b/Engine/Source/Programs/DotNet/EpicGames.Perforce/Records/PrintRecord.cs
--- a/Engine/Source/Programs/DotNet/EpicGames.Perforce/Records/PrintRecord.cs
+++ b/Engine/Source/Programs/DotNet/EpicGames.Perforce/Records/PrintRecord.cs
@@ -6,7 +6,7 @@ using System.Linq;
 using System.Text;
 using System.Threading.Tasks;
 
-namespace Tools.DotNETCommon.Perforce
+namespace EpicGames.Perforce
 {
 	/// <summary>
 	/// Record containing information about a printed file
diff --git a/Engine/Source/Programs/DotNet/EpicGames.Perforce/Records/ReconcileOptions.cs b/Engine/Source/Programs/DotNet/EpicGames.Perforce/Records/ReconcileOptions.cs
--- a/Engine/Source/Programs/DotNet/EpicGames.Perforce/Records/ReconcileOptions.cs
+++ b/Engine/Source/Programs/DotNet/EpicGames.Perforce/Records/ReconcileOptions.cs
@@ -6,7 +6,7 @@ using System.Linq;
 using System.Text;
 using System.Threading.Tasks;
 
-namespace Tools.DotNETCommon.Perforce
+namespace EpicGames.Perforce
 {
 	/// <summary>
 	/// Options for the p4 reconcile command
diff --git a/Engine/Source/Programs/DotNet/EpicGames.Perforce/Records/ReconcileRecord.cs b/Engine/Source/Programs/DotNet/EpicGames.Perforce/Records/ReconcileRecord.cs
--- a/Engine/Source/Programs/DotNet/EpicGames.Perforce/Records/ReconcileRecord.cs
+++ b/Engine/Source/Programs/DotNet/EpicGames.Perforce/Records/ReconcileRecord.cs
@@ -6,7 +6,7 @@ using System.Linq;
 using System.Text;
 using System.Threading.Tasks;
 
-namespace Tools.DotNETCommon.Perforce
+namespace EpicGames.Perforce
 {
 	/// <summary>
 	/// Information about a reconciled file
diff --git a/Engine/Source/Programs/DotNet/EpicGames.Perforce/Records/ResolveOptions.cs b/Engine/Source/Programs/DotNet/EpicGames.Perforce/Records/ResolveOptions.cs
--- a/Engine/Source/Programs/DotNet/EpicGames.Perforce/Records/ResolveOptions.cs
+++ b/Engine/Source/Programs/DotNet/EpicGames.Perforce/Records/ResolveOptions.cs
@@ -6,7 +6,7 @@ using System.Linq;
 using System.Text;
 using System.Threading.Tasks;
 
-namespace Tools.DotNETCommon.Perforce
+namespace EpicGames.Perforce
 {
 	/// <summary>
 	/// Options for the p4 resolve command
diff --git a/Engine/Source/Programs/DotNet/EpicGames.Perforce/Records/ResolveRecord.cs b/Engine/Source/Programs/DotNet/EpicGames.Perforce/Records/ResolveRecord.cs
--- a/Engine/Source/Programs/DotNet/EpicGames.Perforce/Records/ResolveRecord.cs
+++ b/Engine/Source/Programs/DotNet/EpicGames.Perforce/Records/ResolveRecord.cs
@@ -6,7 +6,7 @@ using System.Linq;
 using System.Text;
 using System.Threading.Tasks;
 
-namespace Tools.DotNETCommon.Perforce
+namespace EpicGames.Perforce
 {
 	/// <summary>
 	/// Information about a resolved file
diff --git a/Engine/Source/Programs/DotNet/EpicGames.Perforce/Records/RevertOptions.cs b/Engine/Source/Programs/DotNet/EpicGames.Perforce/Records/RevertOptions.cs
--- a/Engine/Source/Programs/DotNet/EpicGames.Perforce/Records/RevertOptions.cs
+++ b/Engine/Source/Programs/DotNet/EpicGames.Perforce/Records/RevertOptions.cs
@@ -6,7 +6,7 @@ using System.Linq;
 using System.Text;
 using System.Threading.Tasks;
 
-namespace Tools.DotNETCommon.Perforce
+namespace EpicGames.Perforce
 {
 	/// <summary>
 	/// Options for the revert command
diff --git a/Engine/Source/Programs/DotNet/EpicGames.Perforce/Records/RevertRecord.cs b/Engine/Source/Programs/DotNet/EpicGames.Perforce/Records/RevertRecord.cs
--- a/Engine/Source/Programs/DotNet/EpicGames.Perforce/Records/RevertRecord.cs
+++ b/Engine/Source/Programs/DotNet/EpicGames.Perforce/Records/RevertRecord.cs
@@ -6,7 +6,7 @@ using System.Linq;
 using System.Text;
 using System.Threading.Tasks;
 
-namespace Tools.DotNETCommon.Perforce
+namespace EpicGames.Perforce
 {
 	/// <summary>
 	/// Information about a reverted file
diff --git a/Engine/Source/Programs/DotNet/EpicGames.Perforce/Records/RevisionRecord.cs b/Engine/Source/Programs/DotNet/EpicGames.Perforce/Records/RevisionRecord.cs
--- a/Engine/Source/Programs/DotNet/EpicGames.Perforce/Records/RevisionRecord.cs
+++ b/Engine/Source/Programs/DotNet/EpicGames.Perforce/Records/RevisionRecord.cs
@@ -6,7 +6,7 @@ using System.Linq;
 using System.Text;
 using System.Threading.Tasks;
 
-namespace Tools.DotNETCommon.Perforce
+namespace EpicGames.Perforce
 {
 	/// <summary>
 	/// Stores a revision record for a file
diff --git a/Engine/Source/Programs/DotNet/EpicGames.Perforce/Records/ShelveOptions.cs b/Engine/Source/Programs/DotNet/EpicGames.Perforce/Records/ShelveOptions.cs
--- a/Engine/Source/Programs/DotNet/EpicGames.Perforce/Records/ShelveOptions.cs
+++ b/Engine/Source/Programs/DotNet/EpicGames.Perforce/Records/ShelveOptions.cs
@@ -6,7 +6,7 @@ using System.Linq;
 using System.Text;
 using System.Threading.Tasks;
 
-namespace Tools.DotNETCommon.Perforce
+namespace EpicGames.Perforce
 {
 	/// <summary>
 	/// Options for the p4 shelve command
diff --git a/Engine/Source/Programs/DotNet/EpicGames.Perforce/Records/ShelveRecord.cs b/Engine/Source/Programs/DotNet/EpicGames.Perforce/Records/ShelveRecord.cs
--- a/Engine/Source/Programs/DotNet/EpicGames.Perforce/Records/ShelveRecord.cs
+++ b/Engine/Source/Programs/DotNet/EpicGames.Perforce/Records/ShelveRecord.cs
@@ -6,7 +6,7 @@ using System.Linq;
 using System.Text;
 using System.Threading.Tasks;
 
-namespace Tools.DotNETCommon.Perforce
+namespace EpicGames.Perforce
 {
 	/// <summary>
 	/// Contains information about a shelved file
diff --git a/Engine/Source/Programs/DotNet/EpicGames.Perforce/Records/StreamOptions.cs b/Engine/Source/Programs/DotNet/EpicGames.Perforce/Records/StreamOptions.cs
--- a/Engine/Source/Programs/DotNet/EpicGames.Perforce/Records/StreamOptions.cs
+++ b/Engine/Source/Programs/DotNet/EpicGames.Perforce/Records/StreamOptions.cs
@@ -6,7 +6,7 @@ using System.Linq;
 using System.Text;
 using System.Threading.Tasks;
 
-namespace Tools.DotNETCommon.Perforce
+namespace EpicGames.Perforce
 {
 	/// <summary>
 	/// Options for a stream definition
diff --git a/Engine/Source/Programs/DotNet/EpicGames.Perforce/Records/StreamRecord.cs b/Engine/Source/Programs/DotNet/EpicGames.Perforce/Records/StreamRecord.cs
--- a/Engine/Source/Programs/DotNet/EpicGames.Perforce/Records/StreamRecord.cs
+++ b/Engine/Source/Programs/DotNet/EpicGames.Perforce/Records/StreamRecord.cs
@@ -7,7 +7,7 @@ using System.Linq;
 using System.Text;
 using System.Threading.Tasks;
 
-namespace Tools.DotNETCommon.Perforce
+namespace EpicGames.Perforce
 {
 	/// <summary>
 	/// Contains information about a stream, as returned by the 'p4 streams' command
diff --git a/Engine/Source/Programs/DotNet/EpicGames.Perforce/Records/StreamType.cs b/Engine/Source/Programs/DotNet/EpicGames.Perforce/Records/StreamType.cs
--- a/Engine/Source/Programs/DotNet/EpicGames.Perforce/Records/StreamType.cs
+++ b/Engine/Source/Programs/DotNet/EpicGames.Perforce/Records/StreamType.cs
@@ -6,7 +6,7 @@ using System.Linq;
 using System.Text;
 using System.Threading.Tasks;
 
-namespace Tools.DotNETCommon.Perforce
+namespace EpicGames.Perforce
 {
 	/// <summary>
 	/// Type of a Perforce stream
diff --git a/Engine/Source/Programs/DotNet/EpicGames.Perforce/Records/SubmitOptions.cs b/Engine/Source/Programs/DotNet/EpicGames.Perforce/Records/SubmitOptions.cs
--- a/Engine/Source/Programs/DotNet/EpicGames.Perforce/Records/SubmitOptions.cs
+++ b/Engine/Source/Programs/DotNet/EpicGames.Perforce/Records/SubmitOptions.cs
@@ -6,7 +6,7 @@ using System.Linq;
 using System.Text;
 using System.Threading.Tasks;
 
-namespace Tools.DotNETCommon.Perforce
+namespace EpicGames.Perforce
 {
 	/// <summary>
 	/// Options for the 'submit' command
diff --git a/Engine/Source/Programs/DotNet/EpicGames.Perforce/Records/SubmitRecord.cs b/Engine/Source/Programs/DotNet/EpicGames.Perforce/Records/SubmitRecord.cs
--- a/Engine/Source/Programs/DotNet/EpicGames.Perforce/Records/SubmitRecord.cs
+++ b/Engine/Source/Programs/DotNet/EpicGames.Perforce/Records/SubmitRecord.cs
@@ -6,7 +6,7 @@ using System.Linq;
 using System.Text;
 using System.Threading.Tasks;
 
-namespace Tools.DotNETCommon.Perforce
+namespace EpicGames.Perforce
 {
 	/// <summary>
 	/// Contains information about a submitted changelist
diff --git a/Engine/Source/Programs/DotNet/EpicGames.Perforce/Records/SwitchClientOptions.cs b/Engine/Source/Programs/DotNet/EpicGames.Perforce/Records/SwitchClientOptions.cs
--- a/Engine/Source/Programs/DotNet/EpicGames.Perforce/Records/SwitchClientOptions.cs
+++ b/Engine/Source/Programs/DotNet/EpicGames.Perforce/Records/SwitchClientOptions.cs
@@ -6,7 +6,7 @@ using System.Linq;
 using System.Text;
 using System.Threading.Tasks;
 
-namespace Tools.DotNETCommon.Perforce
+namespace EpicGames.Perforce
 {
 	/// <summary>
 	/// Options for the p4 client -s command
diff --git a/Engine/Source/Programs/DotNet/EpicGames.Perforce/Records/SyncOptions.cs b/Engine/Source/Programs/DotNet/EpicGames.Perforce/Records/SyncOptions.cs
--- a/Engine/Source/Programs/DotNet/EpicGames.Perforce/Records/SyncOptions.cs
+++ b/Engine/Source/Programs/DotNet/EpicGames.Perforce/Records/SyncOptions.cs
@@ -6,7 +6,7 @@ using System.Linq;
 using System.Text;
 using System.Threading.Tasks;
 
-namespace Tools.DotNETCommon.Perforce
+namespace EpicGames.Perforce
 {
 	/// <summary>
 	/// Options for the p4 sync command
diff --git a/Engine/Source/Programs/DotNet/EpicGames.Perforce/Records/SyncRecord.cs b/Engine/Source/Programs/DotNet/EpicGames.Perforce/Records/SyncRecord.cs
--- a/Engine/Source/Programs/DotNet/EpicGames.Perforce/Records/SyncRecord.cs
+++ b/Engine/Source/Programs/DotNet/EpicGames.Perforce/Records/SyncRecord.cs
@@ -6,7 +6,7 @@ using System.Linq;
 using System.Text;
 using System.Threading.Tasks;
 
-namespace Tools.DotNETCommon.Perforce
+namespace EpicGames.Perforce
 {
 	/// <summary>
 	/// Information about a synced file
diff --git a/Engine/Source/Programs/DotNet/EpicGames.Perforce/Records/SyncSummaryRecord.cs b/Engine/Source/Programs/DotNet/EpicGames.Perforce/Records/SyncSummaryRecord.cs
--- a/Engine/Source/Programs/DotNet/EpicGames.Perforce/Records/SyncSummaryRecord.cs
+++ b/Engine/Source/Programs/DotNet/EpicGames.Perforce/Records/SyncSummaryRecord.cs
@@ -6,7 +6,7 @@ using System.Linq;
 using System.Text;
 using System.Threading.Tasks;
 
-namespace Tools.DotNETCommon.Perforce
+namespace EpicGames.Perforce
 {
 	/// <summary>
 	/// Information about a sync operation
diff --git a/Engine/Source/Programs/DotNet/EpicGames.Perforce/Records/UnshelveOptions.cs b/Engine/Source/Programs/DotNet/EpicGames.Perforce/Records/UnshelveOptions.cs
--- a/Engine/Source/Programs/DotNet/EpicGames.Perforce/Records/UnshelveOptions.cs
+++ b/Engine/Source/Programs/DotNet/EpicGames.Perforce/Records/UnshelveOptions.cs
@@ -6,7 +6,7 @@ using System.Linq;
 using System.Text;
 using System.Threading.Tasks;
 
-namespace Tools.DotNETCommon.Perforce
+namespace EpicGames.Perforce
 {
 	/// <summary>
 	/// Options for the p4 unshelve command
diff --git a/Engine/Source/Programs/DotNet/EpicGames.Perforce/Records/UnshelveRecord.cs b/Engine/Source/Programs/DotNet/EpicGames.Perforce/Records/UnshelveRecord.cs
--- a/Engine/Source/Programs/DotNet/EpicGames.Perforce/Records/UnshelveRecord.cs
+++ b/Engine/Source/Programs/DotNet/EpicGames.Perforce/Records/UnshelveRecord.cs
@@ -6,7 +6,7 @@ using System.Linq;
 using System.Text;
 using System.Threading.Tasks;
 
-namespace Tools.DotNETCommon.Perforce
+namespace EpicGames.Perforce
 {
 	/// <summary>
 	/// Information about an unshelved file
diff --git a/Engine/Source/Programs/DotNet/EpicGames.Perforce/Records/UpdateChangeOptions.cs b/Engine/Source/Programs/DotNet/EpicGames.Perforce/Records/UpdateChangeOptions.cs
--- a/Engine/Source/Programs/DotNet/EpicGames.Perforce/Records/UpdateChangeOptions.cs
+++ b/Engine/Source/Programs/DotNet/EpicGames.Perforce/Records/UpdateChangeOptions.cs
@@ -6,7 +6,7 @@ using System.Linq;
 using System.Text;
 using System.Threading.Tasks;
 
-namespace Tools.DotNETCommon.Perforce
+namespace EpicGames.Perforce
 {
 	/// <summary>
 	/// Options for the UpdateChange command
diff --git a/Engine/Source/Programs/DotNet/EpicGames.Perforce/Records/UserRecord.cs b/Engine/Source/Programs/DotNet/EpicGames.Perforce/Records/UserRecord.cs
--- a/Engine/Source/Programs/DotNet/EpicGames.Perforce/Records/UserRecord.cs
+++ b/Engine/Source/Programs/DotNet/EpicGames.Perforce/Records/UserRecord.cs
@@ -6,7 +6,7 @@ using System.Linq;
 using System.Text;
 using System.Threading.Tasks;
 
-namespace Tools.DotNETCommon.Perforce
+namespace EpicGames.Perforce
 {
 	/// <summary>
 	/// Contains information about a Perforce user
diff --git a/Engine/Source/Programs/DotNet/EpicGames.Perforce/Records/WhereRecord.cs b/Engine/Source/Programs/DotNet/EpicGames.Perforce/Records/WhereRecord.cs
--- a/Engine/Source/Programs/DotNet/EpicGames.Perforce/Records/WhereRecord.cs
+++ b/Engine/Source/Programs/DotNet/EpicGames.Perforce/Records/WhereRecord.cs
@@ -6,7 +6,7 @@ using System.Linq;
 using System.Text;
 using System.Threading.Tasks;
 
-namespace Tools.DotNETCommon.Perforce
+namespace EpicGames.Perforce
 {
 	/// <summary>
 	/// Contains information about a file's location in the workspace
