Merge JsonObject and JsonWriter to EpicGames.Common as a wrapper around System.Text.Json.

#rnx

diff --git a/Engine/Source/Programs/DotNet/EpicGames.Common/JsonObject.cs b/Engine/Source/Programs/DotNet/EpicGames.Common/JsonObject.cs
new file mode 100644
--- /dev/null
+++ b/Engine/Source/Programs/DotNet/EpicGames.Common/JsonObject.cs
@@ -0,0 +1,482 @@
+// Copyright Epic Games, Inc. All Rights Reserved.
+
+using System;
+using System.Collections.Generic;
+using System.IO;
+using System.Linq;
+using System.Text;
+using System.Text.Json;
+
+namespace EpicGames.Common
+{
+	/// <summary>
+	/// Exception thrown for errors parsing JSON files
+	/// </summary>
+	public class JsonParseException : Exception
+	{
+		/// <summary>
+		/// Constructor
+		/// </summary>
+		/// <param name="Format">Format string</param>
+		/// <param name="Args">Optional arguments</param>
+		public JsonParseException(string Format, params object[] Args)
+			: base(String.Format(Format, Args))
+		{
+		}
+	}
+
+	/// <summary>
+	/// Stores a JSON object in memory
+	/// </summary>
+	public class JsonObject
+	{
+		Dictionary<string, object> RawObject;
+
+		/// <summary>
+		/// Construct a JSON object from the raw string -> object dictionary
+		/// </summary>
+		/// <param name="InRawObject">Raw object parsed from disk</param>
+		public JsonObject(Dictionary<string, object> InRawObject)
+		{
+			RawObject = new Dictionary<string, object>(InRawObject, StringComparer.InvariantCultureIgnoreCase);
+		}
+
+		/// <summary>
+		/// Constructor
+		/// </summary>
+		/// <param name="Element"></param>
+		public JsonObject(JsonElement Element)
+		{
+			RawObject = new Dictionary<string, object>(StringComparer.InvariantCultureIgnoreCase);
+			foreach (JsonProperty Property in Element.EnumerateObject())
+			{
+				RawObject[Property.Name] = ParseElement(Property.Value);
+			}
+		}
+
+		/// <summary>
+		/// Parse an individual element
+		/// </summary>
+		/// <param name="Element"></param>
+		/// <returns></returns>
+		public static object ParseElement(JsonElement Element)
+		{
+			switch(Element.ValueKind)
+			{
+				case JsonValueKind.Array:
+					return Element.EnumerateArray().Select(x => ParseElement(x)).ToArray();
+				case JsonValueKind.Number:
+					return Element.GetDouble();
+				case JsonValueKind.Object:
+					return new JsonObject(Element);
+				case JsonValueKind.String:
+					return Element.GetString();
+				case JsonValueKind.False:
+					return false;
+				case JsonValueKind.True:
+					return true;
+				case JsonValueKind.Null:
+					return null;
+				default:
+					throw new NotImplementedException();
+			}
+		}
+
+		/// <summary>
+		/// Read a JSON file from disk and construct a JsonObject from it
+		/// </summary>
+		/// <param name="File">File to read from</param>
+		/// <returns>New JsonObject instance</returns>
+		public static JsonObject Read(FileReference File)
+		{
+			string Text = FileReference.ReadAllText(File);
+			try
+			{
+				return Parse(Text);
+			}
+			catch(Exception Ex)
+			{
+				throw new JsonParseException("Unable to parse {0}: {1}", File, Ex.Message);
+			}
+		}
+
+		/// <summary>
+		/// Tries to read a JSON file from disk
+		/// </summary>
+		/// <param name="FileName">File to read from</param>
+		/// <param name="Result">On success, receives the parsed object</param>
+		/// <returns>True if the file was read, false otherwise</returns>
+		public static bool TryRead(FileReference FileName, out JsonObject Result)
+		{
+			if (!FileReference.Exists(FileName))
+			{
+				Result = null;
+				return false;
+			}
+
+			string Text = FileReference.ReadAllText(FileName);
+			return TryParse(Text, out Result);
+		}
+
+		/// <summary>
+		/// Parse a JsonObject from the given raw text string
+		/// </summary>
+		/// <param name="Text">The text to parse</param>
+		/// <returns>New JsonObject instance</returns>
+		public static JsonObject Parse(string Text)
+		{
+			JsonDocument Document = JsonSerializer.Deserialize<JsonDocument>(Text);
+			return new JsonObject(Document.RootElement);
+		}
+
+		/// <summary>
+		/// Try to parse a JsonObject from the given raw text string
+		/// </summary>
+		/// <param name="Text">The text to parse</param>
+		/// <param name="Result">On success, receives the new JsonObject</param>
+		/// <returns>True if the object was parsed</returns>
+		public static bool TryParse(string Text, out JsonObject Result)
+		{
+			try
+			{
+				Result = Parse(Text);
+				return true;
+			}
+			catch (Exception)
+			{
+				Result = null;
+				return false;
+			}
+		}
+
+		/// <summary>
+		/// List of key names in this object
+		/// </summary>
+		public IEnumerable<string> KeyNames
+		{
+			get { return RawObject.Keys; }
+		}
+
+		/// <summary>
+		/// Gets a string field by the given name from the object, throwing an exception if it is not there or cannot be parsed.
+		/// </summary>
+		/// <param name="FieldName">Name of the field to get</param>
+		/// <returns>The field value</returns>
+		public string GetStringField(string FieldName)
+		{
+			string StringValue;
+			if (!TryGetStringField(FieldName, out StringValue))
+			{
+				throw new JsonParseException("Missing or invalid '{0}' field", FieldName);
+			}
+			return StringValue;
+		}
+
+		/// <summary>
+		/// Tries to read a string field by the given name from the object
+		/// </summary>
+		/// <param name="FieldName">Name of the field to get</param>
+		/// <param name="Result">On success, receives the field value</param>
+		/// <returns>True if the field could be read, false otherwise</returns>
+		public bool TryGetStringField(string FieldName, out string Result)
+		{
+			object RawValue;
+			if (RawObject.TryGetValue(FieldName, out RawValue) && (RawValue is string))
+			{
+				Result = (string)RawValue;
+				return true;
+			}
+			else
+			{
+				Result = null;
+				return false;
+			}
+		}
+
+		/// <summary>
+		/// Gets a string array field by the given name from the object, throwing an exception if it is not there or cannot be parsed.
+		/// </summary>
+		/// <param name="FieldName">Name of the field to get</param>
+		/// <returns>The field value</returns>
+		public string[] GetStringArrayField(string FieldName)
+		{
+			string[] StringValues;
+			if (!TryGetStringArrayField(FieldName, out StringValues))
+			{
+				throw new JsonParseException("Missing or invalid '{0}' field", FieldName);
+			}
+			return StringValues;
+		}
+
+		/// <summary>
+		/// Tries to read a string array field by the given name from the object
+		/// </summary>
+		/// <param name="FieldName">Name of the field to get</param>
+		/// <param name="Result">On success, receives the field value</param>
+		/// <returns>True if the field could be read, false otherwise</returns>
+		public bool TryGetStringArrayField(string FieldName, out string[] Result)
+		{
+			object RawValue;
+			if (RawObject.TryGetValue(FieldName, out RawValue) && (RawValue is IEnumerable<object>) && ((IEnumerable<object>)RawValue).All(x => x is string))
+			{
+				Result = ((IEnumerable<object>)RawValue).Select(x => (string)x).ToArray();
+				return true;
+			}
+			else
+			{
+				Result = null;
+				return false;
+			}
+		}
+
+		/// <summary>
+		/// Gets a boolean field by the given name from the object, throwing an exception if it is not there or cannot be parsed.
+		/// </summary>
+		/// <param name="FieldName">Name of the field to get</param>
+		/// <returns>The field value</returns>
+		public bool GetBoolField(string FieldName)
+		{
+			bool BoolValue;
+			if (!TryGetBoolField(FieldName, out BoolValue))
+			{
+				throw new JsonParseException("Missing or invalid '{0}' field", FieldName);
+			}
+			return BoolValue;
+		}
+
+		/// <summary>
+		/// Tries to read a bool field by the given name from the object
+		/// </summary>
+		/// <param name="FieldName">Name of the field to get</param>
+		/// <param name="Result">On success, receives the field value</param>
+		/// <returns>True if the field could be read, false otherwise</returns>
+		public bool TryGetBoolField(string FieldName, out bool Result)
+		{
+			object RawValue;
+			if (RawObject.TryGetValue(FieldName, out RawValue) && (RawValue is Boolean))
+			{
+				Result = (bool)RawValue;
+				return true;
+			}
+			else
+			{
+				Result = false;
+				return false;
+			}
+		}
+
+		/// <summary>
+		/// Gets an integer field by the given name from the object, throwing an exception if it is not there or cannot be parsed.
+		/// </summary>
+		/// <param name="FieldName">Name of the field to get</param>
+		/// <returns>The field value</returns>
+		public int GetIntegerField(string FieldName)
+		{
+			int IntegerValue;
+			if (!TryGetIntegerField(FieldName, out IntegerValue))
+			{
+				throw new JsonParseException("Missing or invalid '{0}' field", FieldName);
+			}
+			return IntegerValue;
+		}
+
+		/// <summary>
+		/// Tries to read an integer field by the given name from the object
+		/// </summary>
+		/// <param name="FieldName">Name of the field to get</param>
+		/// <param name="Result">On success, receives the field value</param>
+		/// <returns>True if the field could be read, false otherwise</returns>
+		public bool TryGetIntegerField(string FieldName, out int Result)
+		{
+			object RawValue;
+			if (!RawObject.TryGetValue(FieldName, out RawValue) || !int.TryParse(RawValue.ToString(), out Result))
+			{
+				Result = 0;
+				return false;
+			}
+			return true;
+		}
+
+		/// <summary>
+		/// Tries to read an unsigned integer field by the given name from the object
+		/// </summary>
+		/// <param name="FieldName">Name of the field to get</param>
+		/// <param name="Result">On success, receives the field value</param>
+		/// <returns>True if the field could be read, false otherwise</returns>
+		public bool TryGetUnsignedIntegerField(string FieldName, out uint Result)
+		{
+			object RawValue;
+			if (!RawObject.TryGetValue(FieldName, out RawValue) || !uint.TryParse(RawValue.ToString(), out Result))
+			{
+				Result = 0;
+				return false;
+			}
+			return true;
+		}
+
+		/// <summary>
+		/// Gets a double field by the given name from the object, throwing an exception if it is not there or cannot be parsed.
+		/// </summary>
+		/// <param name="FieldName">Name of the field to get</param>
+		/// <returns>The field value</returns>
+		public double GetDoubleField(string FieldName)
+		{
+			double DoubleValue;
+			if (!TryGetDoubleField(FieldName, out DoubleValue))
+			{
+				throw new JsonParseException("Missing or invalid '{0}' field", FieldName);
+			}
+			return DoubleValue;
+		}
+
+		/// <summary>
+		/// Tries to read a double field by the given name from the object
+		/// </summary>
+		/// <param name="FieldName">Name of the field to get</param>
+		/// <param name="Result">On success, receives the field value</param>
+		/// <returns>True if the field could be read, false otherwise</returns>
+		public bool TryGetDoubleField(string FieldName, out double Result)
+		{
+			object RawValue;
+			if (!RawObject.TryGetValue(FieldName, out RawValue) || !double.TryParse(RawValue.ToString(), out Result))
+			{
+				Result = 0.0;
+				return false;
+			}
+			return true;
+		}
+
+		/// <summary>
+		/// Gets an enum field by the given name from the object, throwing an exception if it is not there or cannot be parsed.
+		/// </summary>
+		/// <param name="FieldName">Name of the field to get</param>
+		/// <returns>The field value</returns>
+		public T GetEnumField<T>(string FieldName) where T : struct
+		{
+			T EnumValue;
+			if (!TryGetEnumField(FieldName, out EnumValue))
+			{
+				throw new JsonParseException("Missing or invalid '{0}' field", FieldName);
+			}
+			return EnumValue;
+		}
+
+		/// <summary>
+		/// Tries to read an enum field by the given name from the object
+		/// </summary>
+		/// <param name="FieldName">Name of the field to get</param>
+		/// <param name="Result">On success, receives the field value</param>
+		/// <returns>True if the field could be read, false otherwise</returns>
+		public bool TryGetEnumField<T>(string FieldName, out T Result) where T : struct
+		{
+			string StringValue;
+			if (!TryGetStringField(FieldName, out StringValue) || !Enum.TryParse<T>(StringValue, true, out Result))
+			{
+				Result = default(T);
+				return false;
+			}
+			return true;
+		}
+
+		/// <summary>
+		/// Tries to read an enum array field by the given name from the object
+		/// </summary>
+		/// <param name="FieldName">Name of the field to get</param>
+		/// <param name="Result">On success, receives the field value</param>
+		/// <returns>True if the field could be read, false otherwise</returns>
+		public bool TryGetEnumArrayField<T>(string FieldName, out T[] Result) where T : struct
+		{
+			string[] StringValues;
+			if (!TryGetStringArrayField(FieldName, out StringValues))
+			{
+				Result = null;
+				return false;
+			}
+
+			T[] EnumValues = new T[StringValues.Length];
+			for (int Idx = 0; Idx < StringValues.Length; Idx++)
+			{
+				if (!Enum.TryParse<T>(StringValues[Idx], true, out EnumValues[Idx]))
+				{
+					Result = null;
+					return false;
+				}
+			}
+
+			Result = EnumValues;
+			return true;
+		}
+
+		/// <summary>
+		/// Gets an object field by the given name from the object, throwing an exception if it is not there or cannot be parsed.
+		/// </summary>
+		/// <param name="FieldName">Name of the field to get</param>
+		/// <returns>The field value</returns>
+		public JsonObject GetObjectField(string FieldName)
+		{
+			JsonObject Result;
+			if (!TryGetObjectField(FieldName, out Result))
+			{
+				throw new JsonParseException("Missing or invalid '{0}' field", FieldName);
+			}
+			return Result;
+		}
+
+		/// <summary>
+		/// Tries to read an object field by the given name from the object
+		/// </summary>
+		/// <param name="FieldName">Name of the field to get</param>
+		/// <param name="Result">On success, receives the field value</param>
+		/// <returns>True if the field could be read, false otherwise</returns>
+		public bool TryGetObjectField(string FieldName, out JsonObject Result)
+		{
+			object RawValue;
+			if (RawObject.TryGetValue(FieldName, out RawValue) && (RawValue is Dictionary<string, object>))
+			{
+				Result = new JsonObject((Dictionary<string, object>)RawValue);
+				return true;
+			}
+			else
+			{
+				Result = null;
+				return false;
+			}
+		}
+
+		/// <summary>
+		/// Gets an object array field by the given name from the object, throwing an exception if it is not there or cannot be parsed.
+		/// </summary>
+		/// <param name="FieldName">Name of the field to get</param>
+		/// <returns>The field value</returns>
+		public JsonObject[] GetObjectArrayField(string FieldName)
+		{
+			JsonObject[] Result;
+			if (!TryGetObjectArrayField(FieldName, out Result))
+			{
+				throw new JsonParseException("Missing or invalid '{0}' field", FieldName);
+			}
+			return Result;
+		}
+
+		/// <summary>
+		/// Tries to read an object array field by the given name from the object
+		/// </summary>
+		/// <param name="FieldName">Name of the field to get</param>
+		/// <param name="Result">On success, receives the field value</param>
+		/// <returns>True if the field could be read, false otherwise</returns>
+		public bool TryGetObjectArrayField(string FieldName, out JsonObject[] Result)
+		{
+			object RawValue;
+			if (RawObject.TryGetValue(FieldName, out RawValue) && (RawValue is IEnumerable<object>) && ((IEnumerable<object>)RawValue).All(x => x is Dictionary<string, object>))
+			{
+				Result = ((IEnumerable<object>)RawValue).Select(x => new JsonObject((Dictionary<string, object>)x)).ToArray();
+				return true;
+			}
+			else
+			{
+				Result = null;
+				return false;
+			}
+		}
+	}
+}
diff --git a/Engine/Source/Programs/DotNet/EpicGames.Common/JsonWriter.cs b/Engine/Source/Programs/DotNet/EpicGames.Common/JsonWriter.cs
new file mode 100644
--- /dev/null
+++ b/Engine/Source/Programs/DotNet/EpicGames.Common/JsonWriter.cs
@@ -0,0 +1,400 @@
+// Copyright Epic Games, Inc. All Rights Reserved.
+
+using System;
+using System.Collections.Generic;
+using System.IO;
+using System.Linq;
+using System.Text;
+using System.Threading.Tasks;
+
+namespace EpicGames.Common
+{
+	/// <summary>
+	/// Specifies how to format JSON output
+	/// </summary>
+	public enum JsonWriterStyle
+	{
+		/// <summary>
+		/// Omit spaces between elements
+		/// </summary>
+		Compact,
+
+		/// <summary>
+		/// Put each value on a newline, and indent output
+		/// </summary>
+		Readable
+	}
+
+	/// <summary>
+	/// Writer for JSON data, which indents the output text appropriately, and adds commas and newlines between fields
+	/// </summary>
+	public class JsonWriter : IDisposable
+	{
+		TextWriter Writer;
+		bool bLeaveOpen;
+		JsonWriterStyle Style;
+		bool bRequiresComma;
+		string Indent;
+
+		/// <summary>
+		/// Constructor
+		/// </summary>
+		/// <param name="FileName">File to write to</param>
+		/// <param name="Style">Should use packed JSON or not</param>
+		public JsonWriter(string FileName, JsonWriterStyle Style = JsonWriterStyle.Readable)
+			: this(new StreamWriter(FileName))
+		{
+			this.Style = Style;
+		}
+
+		/// <summary>
+		/// Constructor
+		/// </summary>
+		/// <param name="FileName">File to write to</param>
+		/// <param name="Style">Should use packed JSON or not</param>
+		public JsonWriter(FileReference FileName, JsonWriterStyle Style = JsonWriterStyle.Readable)
+			: this(new StreamWriter(FileName.FullName))
+		{
+			this.Style = Style;
+		}
+
+		/// <summary>
+		/// Constructor
+		/// </summary>
+		/// <param name="Writer">The text writer to output to</param>
+		/// <param name="bLeaveOpen">Whether to leave the writer open when the object is disposed</param>
+		/// <param name="Style">The output style</param>
+		public JsonWriter(TextWriter Writer, bool bLeaveOpen = false, JsonWriterStyle Style = JsonWriterStyle.Readable)
+		{
+			this.Writer = Writer;
+			this.bLeaveOpen = bLeaveOpen;
+			this.Style = Style;
+			Indent = "";
+		}
+
+		/// <summary>
+		/// Dispose of any managed resources
+		/// </summary>
+		public void Dispose()
+		{
+			if(!bLeaveOpen && Writer != null)
+			{
+				Writer.Dispose();
+				Writer = null;
+			}
+		}
+
+		private void IncreaseIndent()
+		{
+			if (Style == JsonWriterStyle.Readable)
+			{
+				Indent += "\t";
+			}
+		}
+		
+		private void DecreaseIndent()
+		{
+			if (Style == JsonWriterStyle.Readable)
+			{
+				Indent = Indent.Substring(0, Indent.Length - 1);
+			}
+		}
+
+		/// <summary>
+		/// Write the opening brace for an object
+		/// </summary>
+		public void WriteObjectStart()
+		{
+			WriteCommaNewline();
+
+			Writer.Write(Indent);
+			Writer.Write("{");
+
+			IncreaseIndent();
+			bRequiresComma = false;
+		}
+
+		/// <summary>
+		/// Write the name and opening brace for an object
+		/// </summary>
+		/// <param name="ObjectName">Name of the field</param>
+		public void WriteObjectStart(string ObjectName)
+		{
+			WriteCommaNewline();
+
+			string Space = (Style == JsonWriterStyle.Readable) ? " " : "";
+			Writer.Write("{0}\"{1}\":{2}", Indent, ObjectName, Space);
+
+			bRequiresComma = false;
+
+			WriteObjectStart();
+		}
+
+		/// <summary>
+		/// Write the closing brace for an object
+		/// </summary>
+		public void WriteObjectEnd()
+		{
+			DecreaseIndent();
+
+			WriteLine();
+			Writer.Write(Indent);
+			Writer.Write("}");
+
+			bRequiresComma = true;
+		}
+
+		/// <summary>
+		/// Write the opening bracket for an unnamed array
+		/// </summary>
+		/// <param name="ArrayName">Name of the field</param>
+		public void WriteArrayStart()
+		{
+			WriteCommaNewline();
+
+			Writer.Write("{0}[", Indent);
+
+			IncreaseIndent();
+			bRequiresComma = false;
+		}
+
+		/// <summary>
+		/// Write the name and opening bracket for an array
+		/// </summary>
+		/// <param name="ArrayName">Name of the field</param>
+		public void WriteArrayStart(string ArrayName)
+		{
+			WriteCommaNewline();
+
+			string Space = (Style == JsonWriterStyle.Readable) ? " " : "";
+			Writer.Write("{0}\"{1}\":{2}[", Indent, ArrayName, Space);
+
+			IncreaseIndent();
+			bRequiresComma = false;
+		}
+
+		/// <summary>
+		/// Write the closing bracket for an array
+		/// </summary>
+		public void WriteArrayEnd()
+		{
+			DecreaseIndent();
+
+			WriteLine();
+			Writer.Write("{0}]", Indent);
+
+			bRequiresComma = true;
+		}
+
+		private void WriteLine()
+		{
+			if (Style == JsonWriterStyle.Readable)
+			{
+				Writer.WriteLine();
+			}
+		}
+		
+		private void WriteLine(string Line)
+		{
+			if (Style == JsonWriterStyle.Readable)
+			{
+				Writer.WriteLine(Line);
+			}
+			else
+			{
+				Writer.Write(Line);
+			}
+		}
+
+		/// <summary>
+		/// Write an array of strings
+		/// </summary>
+		/// <param name="Name">Name of the field</param>
+		/// <param name="Values">Values for the field</param>
+		public void WriteStringArrayField(string Name, IEnumerable<string> Values)
+		{
+			WriteArrayStart(Name);
+			foreach(string Value in Values)
+			{
+				WriteValue(Value);
+			}
+			WriteArrayEnd();
+		}
+
+		/// <summary>
+		/// Write an array of enum values
+		/// </summary>
+		/// <param name="Name">Name of the field</param>
+		/// <param name="Values">Values for the field</param>
+		public void WriteEnumArrayField<T>(string Name, IEnumerable<T> Values) where T : struct
+		{
+			WriteStringArrayField(Name, Values.Select(x => x.ToString()));
+		}
+
+		/// <summary>
+		/// Write a value with no field name, for the contents of an array
+		/// </summary>
+		/// <param name="Value">Value to write</param>
+		public void WriteValue(int Value)
+		{
+			WriteCommaNewline();
+
+			Writer.Write(Indent);
+			Writer.Write(Value);
+
+			bRequiresComma = true;
+		}
+
+		/// <summary>
+		/// Write a value with no field name, for the contents of an array
+		/// </summary>
+		/// <param name="Value">Value to write</param>
+		public void WriteValue(string Value)
+		{
+			WriteCommaNewline();
+
+			Writer.Write(Indent);
+			WriteEscapedString(Value);
+
+			bRequiresComma = true;
+		}
+
+		/// <summary>
+		/// Write a field name and string value
+		/// </summary>
+		/// <param name="Name">Name of the field</param>
+		/// <param name="Value">Value for the field</param>
+		public void WriteValue(string Name, string Value)
+		{
+			WriteCommaNewline();
+
+			string Space = (Style == JsonWriterStyle.Readable) ? " " : "";
+			Writer.Write("{0}\"{1}\":{2}", Indent, Name, Space);
+			WriteEscapedString(Value);
+
+			bRequiresComma = true;
+		}
+
+		/// <summary>
+		/// Write a field name and integer value
+		/// </summary>
+		/// <param name="Name">Name of the field</param>
+		/// <param name="Value">Value for the field</param>
+		public void WriteValue(string Name, int Value)
+		{
+			WriteValueInternal(Name, Value.ToString());
+		}
+
+		/// <summary>
+		/// Write a field name and double value
+		/// </summary>
+		/// <param name="Name">Name of the field</param>
+		/// <param name="Value">Value for the field</param>
+		public void WriteValue(string Name, double Value)
+		{
+			WriteValueInternal(Name, Value.ToString());
+		}
+
+		/// <summary>
+		/// Write a field name and bool value
+		/// </summary>
+		/// <param name="Name">Name of the field</param>
+		/// <param name="Value">Value for the field</param>
+		public void WriteValue(string Name, bool Value)
+		{
+			WriteValueInternal(Name, Value ? "true" : "false");
+		}
+
+		/// <summary>
+		/// Write a field name and enum value
+		/// </summary>
+		/// <typeparam name="T">The enum type</typeparam>
+		/// <param name="Name">Name of the field</param>
+		/// <param name="Value">Value for the field</param>
+		public void WriteEnumValue<T>(string Name, T Value) where T : struct
+		{
+			WriteValue(Name, Value.ToString());
+		}
+
+		void WriteCommaNewline()
+		{
+			if (bRequiresComma)
+			{
+				WriteLine(",");
+			}
+			else if (Indent.Length > 0)
+			{
+				WriteLine();
+			}
+		}
+
+		void WriteValueInternal(string Name, string Value)
+		{
+			WriteCommaNewline();
+
+			string Space = (Style == JsonWriterStyle.Readable) ? " " : "";
+			Writer.Write("{0}\"{1}\":{2}{3}", Indent, Name, Space, Value);
+
+			bRequiresComma = true;
+		}
+
+		void WriteEscapedString(string Value)
+		{
+			// Escape any characters which may not appear in a JSON string (see http://www.json.org).
+			Writer.Write("\"");
+			if (Value != null)
+			{
+				Writer.Write(EscapeString(Value));
+			}
+			Writer.Write("\"");
+		}
+
+		/// <summary>
+		/// Escapes a string for serializing to JSON
+		/// </summary>
+		/// <param name="Value">The string to escape</param>
+		/// <returns>The escaped string</returns>
+		public static string EscapeString(string Value)
+		{
+			StringBuilder Result = new StringBuilder();
+			for (int Idx = 0; Idx < Value.Length; Idx++)
+			{
+				switch (Value[Idx])
+				{
+					case '\"':
+						Result.Append("\\\"");
+						break;
+					case '\\':
+						Result.Append("\\\\");
+						break;
+					case '\b':
+						Result.Append("\\b");
+						break;
+					case '\f':
+						Result.Append("\\f");
+						break;
+					case '\n':
+						Result.Append("\\n");
+						break;
+					case '\r':
+						Result.Append("\\r");
+						break;
+					case '\t':
+						Result.Append("\\t");
+						break;
+					default:
+						if (Char.IsControl(Value[Idx]))
+						{
+							Result.AppendFormat("\\u{0:X4}", (int)Value[Idx]);
+						}
+						else
+						{
+							Result.Append(Value[Idx]);
+						}
+						break;
+				}
+			}
+			return Result.ToString();
+		}
+	}
+}
