Merging a couple more files to EpicGames.Common.

#rnx

diff --git a/Engine/Source/Programs/DotNet/EpicGames.Common/CsProjectInfo.cs b/Engine/Source/Programs/DotNet/EpicGames.Common/CsProjectInfo.cs
new file mode 100644
--- /dev/null
+++ b/Engine/Source/Programs/DotNet/EpicGames.Common/CsProjectInfo.cs
@@ -0,0 +1,1023 @@
+﻿// Copyright Epic Games, Inc. All Rights Reserved.
+
+using System;
+using System.Collections.Generic;
+using System.IO;
+using System.Linq;
+using System.Runtime.InteropServices;
+using System.Text;
+using System.Threading.Tasks;
+using System.Xml;
+
+namespace EpicGames.Common
+{
+	/// <summary>
+	/// Basic information from a preprocessed C# project file. Supports reading a project file, expanding simple conditions in it, parsing property values, assembly references and references to other projects.
+	/// </summary>
+	public class CsProjectInfo
+	{
+		/// <summary>
+		/// Evaluated properties from the project file
+		/// </summary>
+		public Dictionary<string, string> Properties;
+
+		/// <summary>
+		/// Mapping of referenced assemblies to their 'CopyLocal' (aka 'Private') setting.
+		/// </summary>
+		public Dictionary<FileReference, bool> References = new Dictionary<FileReference, bool>();
+
+		/// <summary>
+		/// Mapping of referenced projects to their 'CopyLocal' (aka 'Private') setting.
+		/// </summary>
+		public Dictionary<FileReference, bool> ProjectReferences = new Dictionary<FileReference, bool>();
+
+		/// <summary>
+		/// List of compile references in the project.
+		/// </summary>
+		public List<FileReference> CompileReferences = new List<FileReference>();
+
+		/// <summary>
+		/// Mapping of content IF they are flagged Always or Newer
+		/// </summary>
+		public Dictionary<FileReference, bool> ContentReferences = new Dictionary<FileReference, bool>();
+
+		/// <summary>
+		/// Path to the CSProject file
+		/// </summary>
+		public FileReference ProjectPath;
+
+		/// <summary>
+		/// Constructor
+		/// </summary>
+		/// <param name="InProperties">Initial mapping of property names to values</param>
+		CsProjectInfo(Dictionary<string, string> InProperties, FileReference InProjectPath)
+		{
+			ProjectPath = InProjectPath;
+			Properties = new Dictionary<string, string>(InProperties);
+		}
+
+		/// <summary>
+		/// Get the ouptut file for this project
+		/// </summary>
+		/// <param name="File">If successful, receives the assembly path</param>
+		/// <returns>True if the output file was found</returns>
+		public bool TryGetOutputFile(out FileReference File)
+		{
+			DirectoryReference OutputDir;
+			if(!TryGetOutputDir(out OutputDir))
+			{
+				File = null;
+				return false;
+			}
+
+			string AssemblyName;
+			if(!TryGetAssemblyName(out AssemblyName))
+			{
+				File = null;
+				return false;
+			}
+
+			File = FileReference.Combine(OutputDir, AssemblyName + ".dll");
+			return true;
+		}
+
+		/// <summary>
+		/// Resolve the project's output directory
+		/// </summary>
+		/// <param name="BaseDirectory">Base directory to resolve relative paths to</param>
+		/// <returns>The configured output directory</returns>
+		public DirectoryReference GetOutputDir(DirectoryReference BaseDirectory)
+		{
+			string OutputPath;
+			if (Properties.TryGetValue("OutputPath", out OutputPath))
+			{
+				return DirectoryReference.Combine(BaseDirectory, OutputPath);
+			}
+			else
+			{
+				return BaseDirectory;
+			}
+		}
+
+		/// <summary>
+		/// Resolve the project's output directory
+		/// </summary>
+		/// <param name="OutputDir">If successful, receives the output directory</param>
+		/// <returns>True if the output directory was found</returns>
+		public bool TryGetOutputDir(out DirectoryReference OutputDir)
+		{
+			string OutputPath;
+			if (Properties.TryGetValue("OutputPath", out OutputPath))
+			{
+				OutputDir = DirectoryReference.Combine(ProjectPath.Directory, OutputPath);
+				return true;
+			}
+			else
+			{
+				OutputDir = null;
+				return false;
+			}
+		}
+
+		/// <summary>
+		/// Returns the assembly name used by this project
+		/// </summary>
+		/// <returns></returns>
+		public bool TryGetAssemblyName(out string AssemblyName)
+		{
+			return Properties.TryGetValue("AssemblyName", out AssemblyName);
+		}
+
+		/// <summary>
+		/// Finds all build products from this project. This includes content and other assemblies marked to be copied local.
+		/// </summary>
+		/// <param name="OutputDir">The output directory</param>
+		/// <param name="BuildProducts">Receives the set of build products</param>
+		/// <param name="ProjectFileToInfo">Map of project file to information, to resolve build products from referenced projects copied locally</param>
+		public void FindBuildProducts(DirectoryReference OutputDir, HashSet<FileReference> BuildProducts, Dictionary<FileReference, CsProjectInfo> ProjectFileToInfo)
+		{
+			// Add the standard build products
+			FindCompiledBuildProducts(OutputDir, BuildProducts);
+
+			// Add the referenced assemblies which are marked to be copied into the output directory. This only happens for the main project, and does not happen for referenced projects.
+			foreach(KeyValuePair<FileReference, bool> Reference in References)
+			{
+				if (Reference.Value)
+				{
+					FileReference OutputFile = FileReference.Combine(OutputDir, Reference.Key.GetFileName());
+					AddReferencedAssemblyAndSupportFiles(OutputFile, BuildProducts);
+				}
+			}
+
+			// Copy the build products for any referenced projects. Note that this does NOT operate recursively.
+			foreach(KeyValuePair<FileReference, bool> ProjectReference in ProjectReferences)
+			{
+				CsProjectInfo OtherProjectInfo;
+				if(ProjectFileToInfo.TryGetValue(ProjectReference.Key, out OtherProjectInfo))
+				{
+					OtherProjectInfo.FindCompiledBuildProducts(OutputDir, BuildProducts);
+				}
+			}
+
+			// Add any copied content. This DOES operate recursively.
+			FindCopiedContent(OutputDir, BuildProducts, ProjectFileToInfo);
+		}
+
+		/// <summary>
+		/// Determines all the compiled build products (executable, etc...) directly built from this project.
+		/// </summary>
+		/// <param name="OutputDir">The output directory</param>
+		/// <param name="BuildProducts">Receives the set of build products</param>
+		public void FindCompiledBuildProducts(DirectoryReference OutputDir, HashSet<FileReference> BuildProducts)
+		{
+			string OutputType, AssemblyName;
+			if (Properties.TryGetValue("OutputType", out OutputType) && Properties.TryGetValue("AssemblyName", out AssemblyName))
+			{
+				switch (OutputType)
+				{
+					case "Exe":
+					case "WinExe":
+						string ExecutableExtension = RuntimeInformation.IsOSPlatform(OSPlatform.Windows) ? ".exe" : "";
+						BuildProducts.Add(FileReference.Combine(OutputDir, AssemblyName + ExecutableExtension));
+						// dotnet outputs a apphost executable and a dll with the actual assembly
+						AddOptionalBuildProduct(FileReference.Combine(OutputDir, AssemblyName + ".dll"), BuildProducts);
+						AddOptionalBuildProduct(FileReference.Combine(OutputDir, AssemblyName + ".pdb"), BuildProducts);
+						AddOptionalBuildProduct(FileReference.Combine(OutputDir, AssemblyName + ".exe.config"), BuildProducts);
+						AddOptionalBuildProduct(FileReference.Combine(OutputDir, AssemblyName + ".exe.mdb"), BuildProducts);
+						break;
+					case "Library":
+						BuildProducts.Add(FileReference.Combine(OutputDir, AssemblyName + ".dll"));
+						AddOptionalBuildProduct(FileReference.Combine(OutputDir, AssemblyName + ".pdb"), BuildProducts);
+						AddOptionalBuildProduct(FileReference.Combine(OutputDir, AssemblyName + ".dll.config"), BuildProducts);
+						AddOptionalBuildProduct(FileReference.Combine(OutputDir, AssemblyName + ".dll.mdb"), BuildProducts);
+						break;
+				}
+			}
+		}
+
+		/// <summary>
+		/// Finds all content which will be copied into the output directory for this project. This includes content from any project references as "copy local" recursively (though MSBuild only traverses a single reference for actual binaries, in such cases)
+		/// </summary>
+		/// <param name="OutputDir">The output directory</param>
+		/// <param name="BuildProducts">Receives the set of build products</param>
+		/// <param name="ProjectFileToInfo">Map of project file to information, to resolve build products from referenced projects copied locally</param>
+		private void FindCopiedContent(DirectoryReference OutputDir, HashSet<FileReference> OutputFiles, Dictionary<FileReference, CsProjectInfo> ProjectFileToInfo)
+		{
+			// Copy any referenced projects too.
+			foreach(KeyValuePair<FileReference, bool> ProjectReference in ProjectReferences)
+			{
+				CsProjectInfo OtherProjectInfo;
+				if(ProjectFileToInfo.TryGetValue(ProjectReference.Key, out OtherProjectInfo))
+				{
+					OtherProjectInfo.FindCopiedContent(OutputDir, OutputFiles, ProjectFileToInfo);
+				}
+			}
+
+			// Add the content which is copied to the output directory
+			foreach (KeyValuePair<FileReference, bool> ContentReference in ContentReferences)
+			{
+				FileReference ContentFile = ContentReference.Key;
+				if (ContentReference.Value)
+				{
+					OutputFiles.Add(FileReference.Combine(OutputDir, ContentFile.GetFileName()));
+				}
+			}
+		}
+
+		/// <summary>
+		/// Adds the given file and any additional build products to the output set
+		/// </summary>
+		/// <param name="OutputFile">The assembly to add</param>
+		/// <param name="OutputFiles">Set to receive the file and support files</param>
+		static public void AddReferencedAssemblyAndSupportFiles(FileReference OutputFile, HashSet<FileReference> OutputFiles)
+		{
+			OutputFiles.Add(OutputFile);
+
+			FileReference SymbolFile = OutputFile.ChangeExtension(".pdb");
+			if (FileReference.Exists(SymbolFile))
+			{
+				OutputFiles.Add(SymbolFile);
+			}
+
+			FileReference DocumentationFile = OutputFile.ChangeExtension(".xml");
+			if (FileReference.Exists(DocumentationFile))
+			{
+				OutputFiles.Add(DocumentationFile);
+			}
+		}
+
+		/// <summary>
+		/// Determines if this project is a .NET core project
+		/// </summary>
+		/// <returns>True if the project is a .NET core project</returns>
+		public bool IsDotNETCoreProject()
+		{
+			bool bIsDotNetCoreProject = false;
+
+			string TargetFramework;
+			if (Properties.TryGetValue("TargetFramework", out TargetFramework))
+			{
+				bIsDotNetCoreProject = TargetFramework.ToLower().Contains("netstandard") || TargetFramework.ToLower().Contains("netcoreapp");
+			}
+
+			return bIsDotNetCoreProject;
+		}
+
+		/// <summary>
+		/// Adds a build product to the output list if it exists
+		/// </summary>
+		/// <param name="BuildProduct">The build product to add</param>
+		/// <param name="BuildProducts">List of output build products</param>
+		public static void AddOptionalBuildProduct(FileReference BuildProduct, HashSet<FileReference> BuildProducts)
+		{
+			if (FileReference.Exists(BuildProduct))
+			{
+				BuildProducts.Add(BuildProduct);
+			}
+		}
+
+		/// <summary>
+		/// Parses supported elements from the children of the provided note. May recurse
+		/// for conditional elements.
+		/// </summary>
+		/// <param name="Node"></param>
+		/// <param name="ProjectInfo"></param>
+		static void ParseNode(XmlNode Node, CsProjectInfo ProjectInfo)
+		{
+			foreach (XmlElement Element in Node.ChildNodes.OfType<XmlElement>())
+			{
+				switch (Element.Name)
+				{
+					case "PropertyGroup":
+						if (EvaluateCondition(Element, ProjectInfo))
+						{
+							ParsePropertyGroup(Element, ProjectInfo);
+						}
+						break;
+					case "ItemGroup":
+						if (EvaluateCondition(Element, ProjectInfo))
+						{
+							ParseItemGroup(ProjectInfo.ProjectPath.Directory, Element, ProjectInfo);
+						}
+						break;
+					case "Choose":
+					case "When":
+						if (EvaluateCondition(Element, ProjectInfo))
+						{
+							ParseNode(Element, ProjectInfo);
+						}
+						break;
+				}
+			}
+		}
+
+		/// <summary>
+		/// Reads project information for the given file.
+		/// </summary>
+		/// <param name="File">The project file to read</param>
+		/// <param name="Properties">Initial set of property values</param>
+		/// <returns>The parsed project info</returns>
+		public static CsProjectInfo Read(FileReference File, Dictionary<string, string> Properties)
+		{
+			CsProjectInfo Project;
+			if(!TryRead(File, Properties, out Project))
+			{
+				throw new Exception(String.Format("Unable to read '{0}'", File));
+			}
+			return Project;
+		}
+
+		/// <summary>
+		/// Attempts to read project information for the given file.
+		/// </summary>
+		/// <param name="File">The project file to read</param>
+		/// <param name="Properties">Initial set of property values</param>
+		/// <param name="OutProjectInfo">If successful, the parsed project info</param>
+		/// <returns>True if the project was read successfully, false otherwise</returns>
+		public static bool TryRead(FileReference File, Dictionary<string, string> Properties, out CsProjectInfo OutProjectInfo)
+		{
+			// Read the project file
+			XmlDocument Document = new XmlDocument();
+			Document.Load(File.FullName);
+
+			// Check the root element is the right type
+			//			HashSet<FileReference> ProjectBuildProducts = new HashSet<FileReference>();
+			if (Document.DocumentElement.Name != "Project")
+			{
+				OutProjectInfo = null;
+				return false;
+			}
+
+			// Parse the basic structure of the document, updating properties and recursing into other referenced projects as we go
+			CsProjectInfo ProjectInfo = new CsProjectInfo(Properties, File);
+
+			// Parse elements in the root node
+			ParseNode(Document.DocumentElement, ProjectInfo);
+
+			// Return the complete project
+			OutProjectInfo = ProjectInfo;
+			return true;
+		}
+
+		/// <summary>
+		/// Parses a 'PropertyGroup' element.
+		/// </summary>
+		/// <param name="ParentElement">The parent 'PropertyGroup' element</param>
+		/// <param name="Properties">Dictionary mapping property names to values</param>
+		static void ParsePropertyGroup(XmlElement ParentElement, CsProjectInfo ProjectInfo)
+		{
+			// We need to know the overridden output type and output path for the selected configuration.
+			foreach (XmlElement Element in ParentElement.ChildNodes.OfType<XmlElement>())
+			{
+				if (EvaluateCondition(Element, ProjectInfo))
+				{
+					ProjectInfo.Properties[Element.Name] = ExpandProperties(Element.InnerText, ProjectInfo.Properties);
+				}
+			}
+		}
+
+		/// <summary>
+		/// Parses an 'ItemGroup' element.
+		/// </summary>
+		/// <param name="BaseDirectory">Base directory to resolve relative paths against</param>
+		/// <param name="ParentElement">The parent 'ItemGroup' element</param>
+		/// <param name="ProjectInfo">Project info object to be updated</param>
+		static void ParseItemGroup(DirectoryReference BaseDirectory, XmlElement ParentElement, CsProjectInfo ProjectInfo)
+		{
+			// Parse any external assembly references
+			foreach (XmlElement ItemElement in ParentElement.ChildNodes.OfType<XmlElement>())
+			{
+				switch (ItemElement.Name)
+				{
+					case "Reference":
+						// Reference to an external assembly
+						if (EvaluateCondition(ItemElement, ProjectInfo))
+						{
+							ParseReference(BaseDirectory, ItemElement, ProjectInfo.References);
+						}
+						break;
+					case "ProjectReference":
+						// Reference to another project
+						if (EvaluateCondition(ItemElement, ProjectInfo))
+						{
+							ParseProjectReference(BaseDirectory, ItemElement, ProjectInfo.Properties, ProjectInfo.ProjectReferences);
+						}
+						break;
+					case "Compile":
+						// Reference to another project
+						if (EvaluateCondition(ItemElement, ProjectInfo))
+						{
+							ParseCompileReference(BaseDirectory, ItemElement, ProjectInfo.CompileReferences);
+						}
+						break;
+					case "Content":
+					case "None":
+						// Reference to another project
+						if (EvaluateCondition(ItemElement, ProjectInfo))
+						{
+							ParseContent(BaseDirectory, ItemElement, ProjectInfo.ContentReferences);
+						}
+						break;
+				}
+			}
+		}
+
+		/// <summary>
+		/// Parses an assembly reference from a given 'Reference' element
+		/// </summary>
+		/// <param name="BaseDirectory">Directory to resolve relative paths against</param>
+		/// <param name="ParentElement">The parent 'Reference' element</param>
+		/// <param name="References">Dictionary of project files to a bool indicating whether the assembly should be copied locally to the referencing project.</param>
+		static void ParseReference(DirectoryReference BaseDirectory, XmlElement ParentElement, Dictionary<FileReference, bool> References)
+		{
+			string HintPath = UnescapeString(GetChildElementString(ParentElement, "HintPath", null));
+			if (!String.IsNullOrEmpty(HintPath))
+			{
+				// Don't include embedded assemblies; they aren't referenced externally by the compiled executable
+				bool bEmbedInteropTypes = GetChildElementBoolean(ParentElement, "EmbedInteropTypes", false);
+				if(!bEmbedInteropTypes)
+				{
+					FileReference AssemblyFile = FileReference.Combine(BaseDirectory, HintPath);
+					bool bPrivate = GetChildElementBoolean(ParentElement, "Private", true);
+					References.Add(AssemblyFile, bPrivate);
+				}
+			}
+		}
+
+		/// <summary>
+		/// Parses a project reference from a given 'ProjectReference' element
+		/// </summary>
+		/// <param name="BaseDirectory">Directory to resolve relative paths against</param>
+		/// <param name="ParentElement">The parent 'ProjectReference' element</param>
+		/// <param name="Properties">Dictionary of properties for parsing the file</param>
+		/// <param name="ProjectReferences">Dictionary of project files to a bool indicating whether the outputs of the project should be copied locally to the referencing project.</param>
+		static void ParseProjectReference(DirectoryReference BaseDirectory, XmlElement ParentElement, Dictionary<string, string> Properties, Dictionary<FileReference, bool> ProjectReferences)
+		{
+			string IncludePath = UnescapeString(ParentElement.GetAttribute("Include"));
+			if (!String.IsNullOrEmpty(IncludePath))
+			{
+				FileReference ProjectFile = FileReference.Combine(BaseDirectory, ExpandProperties(IncludePath, Properties));
+				bool bPrivate = GetChildElementBoolean(ParentElement, "Private", true);
+				ProjectReferences[ProjectFile] = bPrivate;
+			}
+		}
+
+
+		/// recursive helper used by the function below that will append RemainingComponents one by one to ExistingPath, 
+		/// expanding wildcards as necessary. The complete list of files that match the complete path is returned out OutFoundFiles
+		static void ProcessPathComponents(DirectoryReference ExistingPath, IEnumerable<string> RemainingComponents, List<FileReference> OutFoundFiles)
+		{
+			if (!RemainingComponents.Any())
+			{
+				return;
+			}
+
+			// take a look at the first component
+			string CurrentComponent = RemainingComponents.First();
+			RemainingComponents = RemainingComponents.Skip(1);
+
+			// If no other components then this is either a file pattern or a greedy pattern
+			if (!RemainingComponents.Any())
+			{
+				// ** means include all files under this tree, so enumerate them all
+				if (CurrentComponent.Contains("**"))
+				{
+					OutFoundFiles.AddRange(DirectoryReference.EnumerateFiles(ExistingPath, "*", SearchOption.AllDirectories));
+				}
+				else
+				{
+					// easy, a regular path with a file that may or may not be a wildcard
+					OutFoundFiles.AddRange(DirectoryReference.EnumerateFiles(ExistingPath, CurrentComponent));
+				}
+			}
+			else
+			{
+				// new component contains a wildcard, and based on the above we know there are more entries so find 
+				// matching directories
+				if (CurrentComponent.Contains("*"))
+				{
+					// ** means all directories, no matter how deep
+					SearchOption Option = CurrentComponent == "**" ? SearchOption.AllDirectories : SearchOption.TopDirectoryOnly;
+
+					IEnumerable<DirectoryReference> Directories = DirectoryReference.EnumerateDirectories(ExistingPath, CurrentComponent, Option);
+
+					// if we searched all directories regardless of depth, the rest of the components other than the last (file) are irrelevant
+					if (Option == SearchOption.AllDirectories)
+					{
+						RemainingComponents = new[] { RemainingComponents.Last() };
+
+						// ** includes files in the current directory too
+						Directories = Directories.Concat(new[] { ExistingPath });
+					}
+
+					foreach (DirectoryReference Dir in Directories)
+					{
+						ProcessPathComponents(Dir, RemainingComponents, OutFoundFiles);
+					}
+				}
+				else
+				{
+					// add this component to our path and recurse.
+					ExistingPath = DirectoryReference.Combine(ExistingPath, CurrentComponent);
+
+					// but... we can just take all the next components that don't have wildcards in them instead of recursing
+					// into each one!
+					IEnumerable<string> NonWildCardComponents = RemainingComponents.TakeWhile(C => !C.Contains("*"));
+					RemainingComponents = RemainingComponents.Skip(NonWildCardComponents.Count());
+
+					ExistingPath = DirectoryReference.Combine(ExistingPath, NonWildCardComponents.ToArray());
+
+					if (Directory.Exists(ExistingPath.FullName))
+					{
+						ProcessPathComponents(ExistingPath, RemainingComponents, OutFoundFiles);
+					}
+				}
+			}
+		}
+
+		/// <summary>
+		/// Finds all files in the provided path, which may be a csproj wildcard specification.
+		/// E.g. The following are all valid
+		/// Foo/Bar/Item.cs
+		/// Foo/Bar/*.cs
+		/// Foo/*/Item.cs
+		/// Foo/*/*.cs
+		/// Foo/**
+		/// (the last means include all files under the path).
+		/// </summary>
+		/// <param name="InPath">Path specifier to process</param>
+		/// <returns></returns>
+		static IEnumerable<FileReference> FindMatchingFiles(FileReference InPath)
+		{
+			List<FileReference> FoundFiles = new List<FileReference>();		
+
+			// split off the drive root
+			string DriveRoot = Path.GetPathRoot(InPath.FullName);
+
+			// break the rest of the path into components
+			string[] PathComponents = InPath.FullName.Substring(DriveRoot.Length).Split(new char[] { Path.DirectorySeparatorChar, Path.AltDirectorySeparatorChar });
+
+			// Process all the components recursively
+			ProcessPathComponents(new DirectoryReference(DriveRoot), PathComponents, FoundFiles);
+
+			return FoundFiles;
+		}
+
+		/// <summary>
+		/// Parses a project reference from a given 'ProjectReference' element
+		/// </summary>
+		/// <param name="BaseDirectory">Directory to resolve relative paths against</param>
+		/// <param name="ParentElement">The parent 'ProjectReference' element</param>
+		/// <param name="CompileReferences">List of source files.</param>
+		static void ParseCompileReference(DirectoryReference BaseDirectory, XmlElement ParentElement, List<FileReference> CompileReferences)
+		{
+			string IncludePath = UnescapeString(ParentElement.GetAttribute("Include"));
+			if (!String.IsNullOrEmpty(IncludePath))
+			{
+				FileReference SourceFile = FileReference.Combine(BaseDirectory, IncludePath);
+
+				if (SourceFile.FullName.Contains("*"))
+				{
+					CompileReferences.AddRange(FindMatchingFiles(SourceFile));
+				}
+				else
+				{
+					CompileReferences.Add(SourceFile);
+				}
+			}
+		}
+
+		/// <summary>
+		/// Parses an assembly reference from a given 'Content' element
+		/// </summary>
+		/// <param name="BaseDirectory">Directory to resolve relative paths against</param>
+		/// <param name="ParentElement">The parent 'Content' element</param>
+		/// <param name="Contents">Dictionary of project files to a bool indicating whether the assembly should be copied locally to the referencing project.</param>
+		static void ParseContent(DirectoryReference BaseDirectory, XmlElement ParentElement, Dictionary<FileReference, bool> Contents)
+		{
+			string IncludePath = UnescapeString(ParentElement.GetAttribute("Include"));
+			if (!String.IsNullOrEmpty(IncludePath))
+			{
+				string CopyTo = GetChildElementString(ParentElement, "CopyToOutputDirectory", null);
+				bool ShouldCopy = !String.IsNullOrEmpty(CopyTo) && (CopyTo.Equals("Always", StringComparison.InvariantCultureIgnoreCase) || CopyTo.Equals("PreserveNewest", StringComparison.InvariantCultureIgnoreCase));
+				FileReference ContentFile = FileReference.Combine(BaseDirectory, IncludePath);
+				Contents.Add(ContentFile, ShouldCopy);
+			}
+		}
+
+		/// <summary>
+		/// Reads the inner text of a child XML element
+		/// </summary>
+		/// <param name="ParentElement">The parent element to check</param>
+		/// <param name="Name">Name of the child element</param>
+		/// <param name="DefaultValue">Default value to return if the child element is missing</param>
+		/// <returns>The contents of the child element, or default value if it's not present</returns>
+		static string GetChildElementString(XmlElement ParentElement, string Name, string DefaultValue)
+		{
+			XmlElement ChildElement = ParentElement.ChildNodes.OfType<XmlElement>().FirstOrDefault(x => x.Name == Name);
+			if (ChildElement == null)
+			{
+				return DefaultValue;
+			}
+			else
+			{
+				return ChildElement.InnerText ?? DefaultValue;
+			}
+		}
+
+		/// <summary>
+		/// Read a child XML element with the given name, and parse it as a boolean.
+		/// </summary>
+		/// <param name="ParentElement">Parent element to check</param>
+		/// <param name="Name">Name of the child element to look for</param>
+		/// <param name="DefaultValue">Default value to return if the element is missing or not a valid bool</param>
+		/// <returns>The parsed boolean, or the default value</returns>
+		static bool GetChildElementBoolean(XmlElement ParentElement, string Name, bool DefaultValue)
+		{
+			string Value = GetChildElementString(ParentElement, Name, null);
+			if (Value == null)
+			{
+				return DefaultValue;
+			}
+			else if (Value.Equals("True", StringComparison.InvariantCultureIgnoreCase))
+			{
+				return true;
+			}
+			else if (Value.Equals("False", StringComparison.InvariantCultureIgnoreCase))
+			{
+				return false;
+			}
+			else
+			{
+				return DefaultValue;
+			}
+		}
+
+		/// <summary>
+		/// Evaluate whether the optional MSBuild condition on an XML element evaluates to true. Currently only supports 'ABC' == 'DEF' style expressions, but can be expanded as needed.
+		/// </summary>
+		/// <param name="Element">The XML element to check</param>
+		/// <param name="Properties">Dictionary mapping from property names to values.</param>
+		/// <returns></returns>
+		static bool EvaluateCondition(XmlElement Element, CsProjectInfo ProjectInfo)
+		{
+			// Read the condition attribute. If it's not present, assume it evaluates to true.
+			string Condition = Element.GetAttribute("Condition");
+			if (String.IsNullOrEmpty(Condition))
+			{
+				return true;
+			}
+
+			// Expand all the properties
+			Condition = ExpandProperties(Condition, ProjectInfo.Properties);
+
+			// Parse literal true/false values
+			bool OutResult;
+			if (bool.TryParse(Condition, out OutResult))
+			{
+				return OutResult;
+			}
+
+			// Tokenize the condition
+			string[] Tokens = Tokenize(Condition);
+
+			char[] TokenQuotes = new[] { '\'', '(', ')', '{', '}', '[', ']' };
+			
+			// Try to evaluate it. We only support a very limited class of condition expressions at the moment, but it's enough to parse standard projects
+			bool bResult;
+
+			// Handle Exists('Platform\Windows\Gauntlet.TargetDeviceWindows.cs')
+			if (Tokens[0] == "Exists")
+			{
+				// remove all quotes, apostrophes etc that are either tokens or wrap tokens (The Tokenize() function is a bit suspect).
+				string[] Arguments = Tokens.Select(S => S.Trim(TokenQuotes)).Where(S => S.Length > 0).ToArray();
+
+				if (Tokens.Length > 1)
+				{
+					FileSystemReference Dependency = DirectoryReference.Combine(ProjectInfo.ProjectPath.Directory, Arguments[1]);
+
+					if (File.Exists(Dependency.FullName) || Directory.Exists(Dependency.FullName))
+					{
+						return true;
+					}
+
+					return false;
+				}
+			}
+
+			if (Tokens.Length == 3 && Tokens[0].StartsWith("'") && Tokens[1] == "==" && Tokens[2].StartsWith("'"))
+			{
+				bResult = String.Compare(Tokens[0], Tokens[2], StringComparison.InvariantCultureIgnoreCase) == 0;
+			}
+			else if (Tokens.Length == 3 && Tokens[0].StartsWith("'") && Tokens[1] == "!=" && Tokens[2].StartsWith("'"))
+			{
+				bResult = String.Compare(Tokens[0], Tokens[2], StringComparison.InvariantCultureIgnoreCase) != 0;
+			}
+			else
+			{
+				string Msg = string.Format("Couldn't parse condition {0} in project file {1}", Element.ToString(), ProjectInfo.ProjectPath);
+				throw new Exception(Msg);
+			}
+			return bResult;
+		}
+
+		/// <summary>
+		/// Expand MSBuild properties within a string. If referenced properties are not in this dictionary, the process' environment variables are expanded. Unknown properties are expanded to an empty string.
+		/// </summary>
+		/// <param name="Text">The input string to expand</param>
+		/// <param name="Properties">Dictionary mapping from property names to values.</param>
+		/// <returns>String with all properties expanded.</returns>
+		static string ExpandProperties(string Text, Dictionary<string, string> Properties)
+		{
+			string NewText = Text;
+			for (int Idx = NewText.IndexOf("$("); Idx != -1; Idx = NewText.IndexOf("$(", Idx))
+			{
+				// Find the end of the variable name, accounting for changes in scope
+				int EndIdx = Idx + 2;
+				for(int Depth = 1; Depth > 0; EndIdx++)
+				{
+					if(EndIdx == NewText.Length)
+					{
+						throw new Exception("Encountered end of string while expanding properties");
+					}
+					else if(NewText[EndIdx] == '(')
+					{
+						Depth++;
+					}
+					else if(NewText[EndIdx] == ')')
+					{
+						Depth--;
+					}
+				}
+
+				// Convert the property name to tokens
+				string[] Tokens = Tokenize(NewText.Substring(Idx + 2, (EndIdx - 1) - (Idx + 2)));
+
+				// Make sure the first token is a valid property name
+				if(Tokens.Length == 0 || !(Char.IsLetter(Tokens[0][0]) || Tokens[0][0] == '_' || Tokens[0][0] == '[' ))
+				{
+					throw new Exception(String.Format("Invalid property name '{0}' in .csproj file", Tokens[0]));
+				}
+
+				// Find the value for it, either from the dictionary or the environment block
+				string Value;
+				if (!Properties.TryGetValue(Tokens[0], out Value))
+				{
+					Value = Environment.GetEnvironmentVariable(Tokens[0]) ?? "";
+				}
+
+				// Evaluate any functions within it
+				int TokenIdx = 1;
+				while(TokenIdx + 3 < Tokens.Length && Tokens[TokenIdx] == "." && Tokens[TokenIdx + 2] == "(")
+				{
+					// Read the method name
+					string MethodName = Tokens[TokenIdx + 1];
+
+					// Skip to the first argument
+					TokenIdx += 3;
+
+					// Parse any arguments
+					List<object> Arguments = new List<object>();
+					if(Tokens[TokenIdx] != ")")
+					{
+						Arguments.Add(ParseArgument(Tokens[TokenIdx]));
+						TokenIdx++;
+
+						while(TokenIdx + 1 < Tokens.Length && Tokens[TokenIdx] == ",")
+						{
+							Arguments.Add(ParseArgument(Tokens[TokenIdx + 2]));
+							TokenIdx += 2;
+						}
+
+						if(Tokens[TokenIdx] != ")")
+						{
+							throw new Exception("Missing closing parenthesis in condition");
+						}
+					}
+
+					// Skip over the closing parenthesis
+					TokenIdx++;
+
+					// Execute the method
+					try
+					{
+						Value = typeof(string).InvokeMember(MethodName, System.Reflection.BindingFlags.Instance | System.Reflection.BindingFlags.Public | System.Reflection.BindingFlags.InvokeMethod, Type.DefaultBinder, Value, Arguments.ToArray()).ToString();
+					}
+					catch(Exception Ex)
+					{
+						throw new Exception(String.Format("Unable to evaluate condition '{0}'", Text), Ex);
+					}
+				}
+
+				if (TokenIdx < Tokens.Length && Tokens[TokenIdx] == ":")
+				{
+					TokenIdx = Tokens.Length;
+				}
+
+				// Make sure there's nothing left over
+				if(TokenIdx != Tokens.Length)
+				{
+					throw new Exception(String.Format("Unable to parse token '{0}'", NewText));
+				}
+
+				// Replace the variable with its value
+				NewText = NewText.Substring(0, Idx) + Value + NewText.Substring(EndIdx);
+
+				// Make sure we skip over the expanded variable; we don't want to recurse on it.
+				Idx += Value.Length;
+			}
+			return NewText;
+		}
+
+		/// <summary>
+		/// Parse an argument into a framework type
+		/// </summary>
+		/// <param name="Token">The token to parse</param>
+		/// <returns>The argument object</returns>
+		static object ParseArgument(string Token)
+		{
+			// Try to parse a string
+			if(Token.Length > 2 && Token[0] == '\'' && Token[Token.Length - 1] == '\'')
+			{
+				return Token.Substring(1, Token.Length - 2);
+			}
+
+			// Try to parse an integer
+			int Value;
+			if(int.TryParse(Token, out Value))
+			{
+				return Value;
+			}
+
+			// Otherwise throw an exception
+			throw new Exception(String.Format("Unable to parse token '{0}' into a .NET framework type", Token));
+		}
+
+		/// <summary>
+		/// Split an MSBuild condition into tokens
+		/// </summary>
+		/// <param name="Condition">The condition expression</param>
+		/// <returns>Array of the parsed tokens</returns>
+		static string[] Tokenize(string Condition)
+		{
+			List<string> Tokens = new List<string>();
+			for (int Idx = 0; Idx < Condition.Length; )
+			{
+				if(Char.IsWhiteSpace(Condition[Idx]))
+				{
+					// Whitespace
+					Idx++;
+				}
+				else if (Idx + 1 < Condition.Length && Condition[Idx] == '=' && Condition[Idx + 1] == '=')
+				{
+					// "==" operator
+					Idx += 2;
+					Tokens.Add("==");
+				}
+				else if (Idx + 1 < Condition.Length && Condition[Idx] == '!' && Condition[Idx + 1] == '=')
+				{
+					// "!=" operator
+					Idx += 2;
+					Tokens.Add("!=");
+				}
+				else if (Condition[Idx] == '\'')
+				{
+					// Quoted string
+					int StartIdx = Idx++;
+					for(;;Idx++)
+					{
+						if(Idx == Condition.Length)
+						{
+							throw new Exception(String.Format("Missing end quote in condition string ('{0}')", Condition));
+						}
+						if(Condition[Idx] == '\'')
+						{
+							break;
+						}
+					}
+					Idx++;
+					Tokens.Add(Condition.Substring(StartIdx, Idx - StartIdx));
+				}
+				else if (Condition[Idx] == '[')
+				{
+					// static property function invoke
+					// format: [Class]::Property
+					// alternatively: [Class]::Method()
+					// we consider the entire invocation to be a single token
+					int StartIdx = Idx++;
+					int ClassEndIdx = 0;
+					int MethodEndIdx = 0;
+					int MethodArgsEndIdx = 0;
+					for (; ; Idx++)
+					{
+						while (Idx < Condition.Length && (Char.IsLetterOrDigit(Condition[Idx]) || Condition[Idx] == '_'))
+						{
+							Idx++;
+						}
+
+						if (Idx == Condition.Length)
+						{
+							throw new Exception(String.Format("Found end of condition when searching for end of static property function for condition string ('{0}')", Condition));
+						}
+						if (Condition[Idx] == ']')
+						{
+							ClassEndIdx = Idx;
+							Idx++;
+							break;
+						}
+					}
+
+					// skip ::
+					if (Condition[Idx] != ':')
+					{
+						throw new Exception(String.Format("Unexpected format of static property function, expected :: after class declaration in condition string ('{0}')", Condition));
+					}
+					Idx += 2;
+
+					while (Idx < Condition.Length && (Char.IsLetterOrDigit(Condition[Idx]) || Condition[Idx] == '_'))
+					{
+						Idx++;
+					}
+
+					MethodEndIdx = Idx;
+
+					if (Condition[Idx] == '(')
+					{
+						// a method invoke
+						for (; ; Idx++)
+						{
+							while (Idx < Condition.Length && (Char.IsLetterOrDigit(Condition[Idx]) || Condition[Idx] == '_'))
+							{
+								Idx++;
+							}
+
+							if (Idx == Condition.Length)
+							{
+								throw new Exception(String.Format("Found end of condition when searching for ) to indicate end of arguments to static property function for condition string ('{0}')", Condition));
+							}
+							if (Condition[Idx] == ')')
+							{
+								MethodArgsEndIdx = Idx;
+								Idx++;
+								break;
+							}
+						}
+
+					}
+
+					Idx++;
+					Tokens.Add(Condition.Substring(StartIdx, Idx - StartIdx));
+				}
+				else if(Char.IsLetterOrDigit(Condition[Idx]) || Condition[Idx] == '_')
+				{
+					// Identifier or number
+					int StartIdx = Idx++;
+					while(Idx < Condition.Length && (Char.IsLetterOrDigit(Condition[Idx]) || Condition[Idx] == '_'))
+					{
+						Idx++;
+					}
+					Tokens.Add(Condition.Substring(StartIdx, Idx - StartIdx));
+				}
+				else
+				{
+					// Other token; assume a single character.
+					string Token = Condition.Substring(Idx++, 1);
+					Tokens.Add(Token);
+				}
+			}
+			return Tokens.ToArray();
+		}
+
+		/// <summary>
+		/// Un-escape an MSBuild string (see https://msdn.microsoft.com/en-us/library/bb383819.aspx)
+		/// </summary>
+		/// <param name="Text">String to remove escape characters from</param>
+		/// <returns>Unescaped string</returns>
+		static string UnescapeString(string Text)
+		{
+			const string HexChars = "0123456789abcdef";
+
+			string NewText = Text;
+			if(NewText != null)
+			{
+				for(int Idx = 0; Idx + 2 < NewText.Length; Idx++)
+				{
+					if(NewText[Idx] == '%')
+					{
+						int UpperDigitIdx = HexChars.IndexOf(Char.ToLowerInvariant(NewText[Idx + 1]));
+						if(UpperDigitIdx != -1)
+						{
+							int LowerDigitIdx = HexChars.IndexOf(Char.ToLowerInvariant(NewText[Idx + 2]));
+							if(LowerDigitIdx != -1)
+							{
+								char NewChar = (char)((UpperDigitIdx << 4) | LowerDigitIdx);
+								NewText = NewText.Substring(0, Idx) + NewChar + NewText.Substring(Idx + 3);
+							}
+						}
+					}
+				}
+			}
+			return NewText;
+		}
+	}
+}
diff --git a/Engine/Source/Programs/DotNet/EpicGames.Common/TraceSpan.cs b/Engine/Source/Programs/DotNet/EpicGames.Common/TraceSpan.cs
new file mode 100644
--- /dev/null
+++ b/Engine/Source/Programs/DotNet/EpicGames.Common/TraceSpan.cs
@@ -0,0 +1,135 @@
+// Copyright Epic Games, Inc. All Rights Reserved.
+
+using System;
+using System.Collections.Generic;
+using System.Diagnostics;
+using System.Globalization;
+using System.IO;
+using System.Linq;
+using System.Reflection;
+using System.Text;
+using System.Threading.Tasks;
+
+namespace EpicGames.Common
+{
+	/// <summary>
+	/// Interface for a trace span
+	/// </summary>
+	public interface ITraceSpan : IDisposable
+	{
+		/// <summary>
+		/// Adds additional metadata to this scope
+		/// </summary>
+		/// <param name="Name">Name of the key</param>
+		/// <param name="Value">Value for this metadata</param>
+		void AddMetadata(string Name, string Value);
+	}
+
+	/// <summary>
+	/// Methods for creating ITraceScope instances
+	/// </summary>
+	public static class TraceSpan
+	{
+		/// <summary>
+		/// Concrete implementation of ITraceScope
+		/// </summary>
+		class TraceSpanImpl : ITraceSpan
+		{
+			public string Name;
+			public string Resource;
+			public string Service;
+			public DateTimeOffset StartTime;
+			public DateTimeOffset? FinishTime;
+			public Dictionary<string, string> Metadata = new Dictionary<string, string>();
+
+			public TraceSpanImpl(string Name, string Resource, string Service)
+			{
+				this.Name = Name;
+				this.Resource = Resource;
+				this.Service = Service;
+				this.StartTime = DateTimeOffset.Now;
+			}
+
+			public void AddMetadata(string Name, string Value)
+			{
+				Metadata[Name] = Value;
+			}
+
+			public void Dispose()
+			{
+				if(FinishTime == null)
+				{
+					FinishTime = DateTimeOffset.Now;
+				}
+			}
+		}
+
+		/// <summary>
+		/// The current scope provider
+		/// </summary>
+		static List<TraceSpanImpl> Spans = new List<TraceSpanImpl>();
+
+		/// <summary>
+		/// Creates a scope using the current provider
+		/// </summary>
+		public static ITraceSpan Create(string Name, string Resource = null, string Service = null)
+		{
+			TraceSpanImpl Span = new TraceSpanImpl(Name, Resource, Service);
+			Spans.Add(Span);
+			return Span;
+		}
+
+		/// <summary>
+		/// Saves all the scope information to a file
+		/// </summary>
+		public static void Flush()
+		{
+			string TelemetryDir = Environment.GetEnvironmentVariable("UE_TELEMETRY_DIR");
+			if (TelemetryDir != null)
+			{
+				FileReference File;
+				using (Process Process = Process.GetCurrentProcess())
+				{
+					DirectoryReference TelemetryDirRef = new DirectoryReference(TelemetryDir);
+					DirectoryReference.CreateDirectory(TelemetryDirRef);
+
+					string FileName = String.Format("{0}.{1}.json", Path.GetFileName(Assembly.GetEntryAssembly().Location), Process.Id, Process.StartTime.Ticks);
+					File = FileReference.Combine(TelemetryDirRef, FileName);
+				}
+
+				using (JsonWriter Writer = new JsonWriter(File))
+				{
+					Writer.WriteObjectStart();
+					Writer.WriteArrayStart("Spans");
+					foreach (TraceSpanImpl Span in Spans)
+					{
+						if (Span.FinishTime != null)
+						{
+							Writer.WriteObjectStart();
+							Writer.WriteValue("Name", Span.Name);
+							if (Span.Resource != null)
+							{
+								Writer.WriteValue("Resource", Span.Resource);
+							}
+							if (Span.Service != null)
+							{
+								Writer.WriteValue("Service", Span.Service);
+							}
+							Writer.WriteValue("StartTime", Span.StartTime.ToString("o", CultureInfo.InvariantCulture));
+							Writer.WriteValue("FinishTime", Span.FinishTime.Value.ToString("o", CultureInfo.InvariantCulture));
+							Writer.WriteObjectStart("Metadata");
+							foreach (KeyValuePair<string, string> Pair in Span.Metadata)
+							{
+								Writer.WriteValue(Pair.Key, Pair.Value);
+							}
+							Writer.WriteObjectEnd();
+							Writer.WriteObjectEnd();
+						}
+					}
+					Writer.WriteArrayEnd();
+					Writer.WriteObjectEnd();
+				}
+			}
+		}
+	}
+}
