UBT: Remove more dependencies on fastJSON.

#rnx

diff --git a/Engine/Source/Programs/DotNet/EpicGames.Common/BinaryFormatterUtils.cs b/Engine/Source/Programs/DotNet/EpicGames.Common/BinaryFormatterUtils.cs
new file mode 100644
--- /dev/null
+++ b/Engine/Source/Programs/DotNet/EpicGames.Common/BinaryFormatterUtils.cs
@@ -0,0 +1,66 @@
+// Copyright Epic Games, Inc. All Rights Reserved.
+
+using System;
+using System.Collections.Generic;
+using System.IO;
+using System.Linq;
+using System.Runtime.Serialization.Formatters.Binary;
+using System.Text;
+using System.Threading.Tasks;
+
+namespace EpicGames.Common
+{
+	/// <summary>
+	/// Utility functions for serializing using the BinaryFormatter
+	/// </summary>
+	public static class BinaryFormatterUtils
+	{
+		/// <summary>
+		/// Load an object from a file on disk, using the binary formatter
+		/// </summary>
+		/// <param name="Location">File to read from</param>
+		/// <returns>Instance of the object that was read from disk</returns>
+		public static T Load<T>(FileReference Location)
+		{
+			using (FileStream Stream = new FileStream(Location.FullName, FileMode.Open, FileAccess.Read))
+			{
+				BinaryFormatter Formatter = new BinaryFormatter();
+				return (T)Formatter.Deserialize(Stream);
+			}
+		}
+
+		/// <summary>
+		/// Saves a file to disk, using the binary formatter
+		/// </summary>
+		/// <param name="Location">File to write to</param>
+		/// <param name="Object">Object to serialize</param>
+		public static void Save(FileReference Location, object Object)
+		{
+			DirectoryReference.CreateDirectory(Location.Directory);
+			using (FileStream Stream = new FileStream(Location.FullName, FileMode.Create, FileAccess.Write))
+			{
+				BinaryFormatter Formatter = new BinaryFormatter();
+				Formatter.Serialize(Stream, Object);
+			}
+		}
+
+		/// <summary>
+		/// Saves a file to disk using the binary formatter, without updating the timestamp if it hasn't changed
+		/// </summary>
+		/// <param name="Location">File to write to</param>
+		/// <param name="Object">Object to serialize</param>
+		public static void SaveIfDifferent(FileReference Location, object Object)
+		{
+			byte[] Contents;
+			using(MemoryStream Stream = new MemoryStream())
+			{
+				BinaryFormatter Formatter = new BinaryFormatter();
+				Formatter.Serialize(Stream, Object);
+				Contents = Stream.ToArray();
+			}
+
+			DirectoryReference.CreateDirectory(Location.Directory);
+			FileReference.WriteAllBytesIfDifferent(Location, Contents);
+		}
+	}
+}
diff --git a/Engine/Source/Programs/UnrealBuildTool/Executors/SNDBS.cs b/Engine/Source/Programs/UnrealBuildTool/Executors/SNDBS.cs
--- a/Engine/Source/Programs/UnrealBuildTool/Executors/SNDBS.cs
+++ b/Engine/Source/Programs/UnrealBuildTool/Executors/SNDBS.cs
@@ -7,6 +7,7 @@ using System.Diagnostics;
 using System.IO;
 using System.Linq;
 using System.ServiceProcess;
+using System.Text.Json;
 using Tools.DotNETCommon;
 
 #nullable disable
@@ -65,7 +66,7 @@ namespace UnrealBuildTool
 
 			// Build the json script file to describe all the actions and their dependencies
 			var ActionIds = Actions.ToDictionary(a => a, a => Guid.NewGuid().ToString());
-			File.WriteAllText(ScriptFile.FullName, Json.Serialize(new Dictionary<string, object>()
+			File.WriteAllText(ScriptFile.FullName, JsonSerializer.Serialize(new Dictionary<string, object>()
 			{
 				["jobs"] = Actions.ToDictionary(a => ActionIds[a], a =>
 				{
