Fix UAT compile error due to incorrect namespace.

#rnx

diff --git a/Engine/Source/Programs/DotNet/EpicGames.Common/HelpUtils.cs b/Engine/Source/Programs/DotNet/EpicGames.Common/HelpUtils.cs
--- a/Engine/Source/Programs/DotNet/EpicGames.Common/HelpUtils.cs
+++ b/Engine/Source/Programs/DotNet/EpicGames.Common/HelpUtils.cs
@@ -7,7 +7,7 @@ using System.Linq;
 using System.Text;
 using System.Threading.Tasks;
 
-namespace Tools.DotNETCommon
+namespace EpicGames.Common
 {
 	/// <summary>
 	/// Utility functions for showing help for objects
